OodleDataCompression in Core
All games now use Engine Oodle for runtime decompression
Oodle in Engine can load an older Oodle DLL for encoding packages to keep making the same binary data

diff --git a/Engine/Plugins/Compression/OodleData/OodleData.uplugin b/Engine/Plugins/Compression/OodleData/OodleData.uplugin
deleted file mode 100644
diff --git a/Engine/Plugins/Compression/OodleData/Sdks/2.9.2/include/oodle2.h b/Engine/Plugins/Compression/OodleData/Sdks/2.9.2/include/oodle2.h
deleted file mode 100644
diff --git a/Engine/Plugins/Compression/OodleData/Sdks/2.9.2/include/oodle2base.h b/Engine/Plugins/Compression/OodleData/Sdks/2.9.2/include/oodle2base.h
deleted file mode 100644
diff --git a/Engine/Plugins/Compression/OodleData/Source/OodleDataCompressionFormat.Build.cs b/Engine/Plugins/Compression/OodleData/Source/OodleDataCompressionFormat.Build.cs
deleted file mode 100644
diff --git a/Engine/Plugins/Compression/OodleData/Source/Private/OodleDataCompressionFormat.cpp b/Engine/Plugins/Compression/OodleData/Source/Private/OodleDataCompressionFormat.cpp
deleted file mode 100644
diff --git a/Engine/Plugins/Compression/OodleData/Source/Private/OodleDataCompressionFormatPCH.h b/Engine/Plugins/Compression/OodleData/Source/Private/OodleDataCompressionFormatPCH.h
deleted file mode 100644
diff --git a/Engine/Plugins/Compression/OodleNetwork/Source/Private/OodleNetworkArchives.cpp b/Engine/Plugins/Compression/OodleNetwork/Source/Private/OodleNetworkArchives.cpp
--- a/Engine/Plugins/Compression/OodleNetwork/Source/Private/OodleNetworkArchives.cpp
+++ b/Engine/Plugins/Compression/OodleNetwork/Source/Private/OodleNetworkArchives.cpp
@@ -2,6 +2,7 @@
 
 #include "OodleNetworkArchives.h"
 #include "OodleNetworkHandlerComponent.h"
+#include "Compression/OodleDataCompression.h"
 
 // Maximum size of compress/decompress buffers (just under 2GB, due to max int32 value)
 #define MAX_COMPRESS_BUFFER (1024 * 1024 * 2047)
@@ -61,23 +62,54 @@ bool FOodleNetworkArchiveBase::SerializeOodleDecompressData(FOodleCompressedData
 	bSuccess = ensure(CompressedLength <= (InnerArchive.TotalSize() - DataOffset));
 	bSuccess = bSuccess && ensure(DecompressedLength <= MAX_COMPRESS_BUFFER);
 	bSuccess = bSuccess && ensure(CompressedLength <= MAX_COMPRESS_BUFFER);
-
-	// compression no longer supported here :
-	bSuccess = bSuccess && ensure(CompressedLength == DecompressedLength);
-
+	
 	if (bSuccess)
 	{
 		SeekPush(DataOffset);
 
 		uint8* DecompressedData = new uint8[DecompressedLength];
 
-		InnerArchive.Serialize(DecompressedData, DecompressedLength);
+		if (CompressedLength == DecompressedLength)
+		{
+			// uncompressed
 
-		OutData = DecompressedData;
-		OutDataBytes = DecompressedLength;
+			InnerArchive.Serialize(DecompressedData, DecompressedLength);
+		}
+		else
+		{
+			// all data should be uncompressed now
+			// but support decompressing legacy data here
+
+			uint8* CompressedData = new uint8[CompressedLength];
+		
+			InnerArchive.Serialize(CompressedData, CompressedLength);
+
+			if (!InnerArchive.IsError())
+			{
+				bSuccess = OodleDataDecompress(
+					(void*)DecompressedData, DecompressedLength,
+					(void*)CompressedData, CompressedLength);
+			}
+
+			delete[] CompressedData;
+		}
+
+		bSuccess = bSuccess && !InnerArchive.IsError();
+
+		if (bSuccess)
+		{
+			OutData = DecompressedData;
+			OutDataBytes = DecompressedLength;
+		}
+		else
+		{
+			delete[] DecompressedData;
+		}
 
 		SeekPop();
 	}
+	
+	bSuccess = bSuccess && !InnerArchive.IsError();
 
 	return bSuccess;
 }
diff --git a/Engine/Plugins/Compression/OodleNetwork/Source/Private/OodleNetworkHandlerComponent.cpp b/Engine/Plugins/Compression/OodleNetwork/Source/Private/OodleNetworkHandlerComponent.cpp
--- a/Engine/Plugins/Compression/OodleNetwork/Source/Private/OodleNetworkHandlerComponent.cpp
+++ b/Engine/Plugins/Compression/OodleNetwork/Source/Private/OodleNetworkHandlerComponent.cpp
@@ -305,8 +305,8 @@ FOodleNetworkDictionary::~FOodleNetworkDictionary()
 OodleNetworkHandlerComponent::OodleNetworkHandlerComponent()
 	: HandlerComponent(FName(TEXT("OodleNetworkHandlerComponent")))
 	, bEnableOodle(false)
-	, ServerEnableMode(EOodleEnableMode::AlwaysEnabled)
-	, ClientEnableMode(EOodleEnableMode::AlwaysEnabled)
+	, ServerEnableMode(EOodleNetworkEnableMode::AlwaysEnabled)
+	, ClientEnableMode(EOodleNetworkEnableMode::AlwaysEnabled)
 #if !UE_BUILD_SHIPPING || OODLE_DEV_SHIPPING
 	, InPacketLog(nullptr)
 	, OutPacketLog(nullptr)
@@ -394,10 +394,10 @@ void OodleNetworkHandlerComponent::Initialize()
 	}
 
 	FString CurEnableModeStr;
-	UEnum* EnableModeEnum = StaticEnum<EOodleEnableMode>();
+	UEnum* EnableModeEnum = StaticEnum<EOodleNetworkEnableMode>();
 	bool bSetServerEnableMode = false;
 	bool bSetClientEnableMode = false;
-	auto SetEnableModeFromStr = [&EnableModeEnum](EOodleEnableMode& OutMode, FString ModeStr) -> bool
+	auto SetEnableModeFromStr = [&EnableModeEnum](EOodleNetworkEnableMode& OutMode, FString ModeStr) -> bool
 		{
 			bool bSuccess = false;
 
@@ -405,12 +405,12 @@ void OodleNetworkHandlerComponent::Initialize()
 
 			if (EnumVal != INDEX_NONE)
 			{
-				OutMode = (EOodleEnableMode)EnumVal;
+				OutMode = (EOodleNetworkEnableMode)EnumVal;
 				bSuccess = true;
 			}
 			else
 			{
-				UE_LOG(OodleNetworkHandlerComponentLog, Error, TEXT("Failed to parse EOodleEnableMode value '%s'"), *ModeStr);
+				UE_LOG(OodleNetworkHandlerComponentLog, Error, TEXT("Failed to parse EOodleNetworkEnableMode value '%s'"), *ModeStr);
 			}
 
 			return bSuccess;
@@ -482,9 +482,9 @@ void OodleNetworkHandlerComponent::Initialize()
 
 #endif
 
-		EOodleEnableMode EnableMode = (Handler->Mode == Handler::Mode::Server ? ServerEnableMode : ClientEnableMode);
+		EOodleNetworkEnableMode EnableMode = (Handler->Mode == Handler::Mode::Server ? ServerEnableMode : ClientEnableMode);
 
-		if (EnableMode == EOodleEnableMode::AlwaysEnabled)
+		if (EnableMode == EOodleNetworkEnableMode::AlwaysEnabled)
 		{
 			InitializeDictionaries();
 		}
@@ -954,8 +954,8 @@ void OodleNetworkHandlerComponent::Incoming(FBitReader& Packet)
 		{
 			const bool bIsServer = (Handler->Mode == Handler::Mode::Server);
 
-			// Lazy-loading of dictionary when EOodleEnableMode::WhenCompressedPacketReceived is active
-			if (!bInitializedDictionaries && (bIsServer ? ServerEnableMode : ClientEnableMode) == EOodleEnableMode::WhenCompressedPacketReceived)
+			// Lazy-loading of dictionary when EOodleNetworkEnableMode::WhenCompressedPacketReceived is active
+			if (!bInitializedDictionaries && (bIsServer ? ServerEnableMode : ClientEnableMode) == EOodleNetworkEnableMode::WhenCompressedPacketReceived)
 			{
 				RemoteInitializeDictionaries();
 			}
@@ -1144,7 +1144,7 @@ void OodleNetworkHandlerComponent::Outgoing(FBitWriter& Packet, FOutPacketTraits
 #endif
 
 		// Skip compression when we are waiting on the remote side to enable compression
-		if (!bSkipCompression && (!bInitializedDictionaries && (bIsServer ? ServerEnableMode : ClientEnableMode) == EOodleEnableMode::WhenCompressedPacketReceived))
+		if (!bSkipCompression && (!bInitializedDictionaries && (bIsServer ? ServerEnableMode : ClientEnableMode) == EOodleNetworkEnableMode::WhenCompressedPacketReceived))
 		{
 			bSkipCompression = true;
 			bSkipCompressionClientDisabled = true;
@@ -1342,7 +1342,7 @@ bool OodleNetworkHandlerComponent::IsCompressionActive() const
 			(CurDict == nullptr && bCaptureMode) || bOodleCompressionDisabled ||
 #endif
 			// Skip compression when we are waiting on the remote side to enable compression
-			(!bInitializedDictionaries && (bIsServer ? ServerEnableMode : ClientEnableMode) == EOodleEnableMode::WhenCompressedPacketReceived)
+			(!bInitializedDictionaries && (bIsServer ? ServerEnableMode : ClientEnableMode) == EOodleNetworkEnableMode::WhenCompressedPacketReceived)
 			;
 
 		if (CurDict != nullptr && !bSkipCompression)
diff --git a/Engine/Plugins/Compression/OodleNetwork/Source/Private/OodleNetworkUtils.cpp b/Engine/Plugins/Compression/OodleNetwork/Source/Private/OodleNetworkUtils.cpp
deleted file mode 100644
diff --git a/Engine/Plugins/Compression/OodleNetwork/Source/Public/OodleNetworkHandlerComponent.h b/Engine/Plugins/Compression/OodleNetwork/Source/Public/OodleNetworkHandlerComponent.h
--- a/Engine/Plugins/Compression/OodleNetwork/Source/Public/OodleNetworkHandlerComponent.h
+++ b/Engine/Plugins/Compression/OodleNetwork/Source/Public/OodleNetworkHandlerComponent.h
@@ -25,7 +25,7 @@ DECLARE_LOG_CATEGORY_EXTERN(OodleNetworkHandlerComponentLog, Log, All);
  * Specifies when compression is enabled. Used to make compression optional, for some platforms/clients
  */
 UENUM()
-enum class EOodleEnableMode : uint8
+enum class EOodleNetworkEnableMode : uint8
 {
 	AlwaysEnabled,					// Oodle compression is always enabled - forces compression to be enabled remotely
 	WhenCompressedPacketReceived	// Oodle compression is only enabled if remotely requested
@@ -369,10 +369,10 @@ protected:
 	bool bEnableOodle;
 
 	/** When to enable compression on the server */
-	EOodleEnableMode ServerEnableMode;
+	EOodleNetworkEnableMode ServerEnableMode;
 
 	/** When to enable compression on the client */
-	EOodleEnableMode ClientEnableMode;
+	EOodleNetworkEnableMode ClientEnableMode;
 
 #if !UE_BUILD_SHIPPING || OODLE_DEV_SHIPPING
 	/** File to log input packets to */
diff --git a/Engine/Plugins/Compression/OodleNetwork/Source/Public/OodleNetworkUtils.h b/Engine/Plugins/Compression/OodleNetwork/Source/Public/OodleNetworkUtils.h
deleted file mode 100644
diff --git a/Engine/Source/Runtime/Core/Core.Build.cs b/Engine/Source/Runtime/Core/Core.Build.cs
--- a/Engine/Source/Runtime/Core/Core.Build.cs
+++ b/Engine/Source/Runtime/Core/Core.Build.cs
@@ -20,6 +20,7 @@ public class Core : ModuleRules
 		}
 
 		PrivateDependencyModuleNames.Add("BLAKE3");
+		PrivateDependencyModuleNames.Add("OodleDataCompression");
 
 		PublicDependencyModuleNames.Add("TraceLog");
 		PublicIncludePaths.Add("Runtime/TraceLog/Public");
diff --git a/Engine/Source/Runtime/Core/Private/Compression/OodleDataCompression.cpp b/Engine/Source/Runtime/Core/Private/Compression/OodleDataCompression.cpp
new file mode 100644
--- /dev/null
+++ b/Engine/Source/Runtime/Core/Private/Compression/OodleDataCompression.cpp
@@ -0,0 +1,278 @@
+// Copyright Epic Games, Inc. All Rights Reserved.
+
+#include "Compression/OodleDataCompression.h"
+#include "Misc/ConfigCacheIni.h"
+#include "Misc/ICompressionFormat.h"
+#include "Templates/CheckValueCast.h"
+
+#include "oodle2.h"
+
+DEFINE_LOG_CATEGORY(OodleDataCompression);
+
+struct OodleDataCompressionDecoders
+{
+	OO_SINTa OodleDecoderMemorySize = 0;
+	int32 OodleDecoderBufferCount =0;
+
+	struct OodleDataCompressionDecoder
+	{
+		FCriticalSection OodleDecoderMutex;
+		void * OodleDecoderMemory = nullptr;
+		char pad[64];
+
+		OodleDataCompressionDecoder()
+		{
+		}
+	};
+
+	OodleDataCompressionDecoder* OodleDecoders = nullptr;
+
+	OodleDataCompressionDecoders()
+	{
+		// enough decoder scratch for any compressor & buffer size.
+		// note "InCompressor" is what we want to Encode with but we may be asked to decode other compressors!
+		OodleDecoderMemorySize = OodleLZDecoder_MemorySizeNeeded(OodleLZ_Compressor_Invalid, -1);
+
+		int32 BufferCount = 2;
+
+		// be wary of possible init order problem
+		//  if we init Oodle before config might not exist yet?
+		//  can we just check GConfig vs nullptr ?
+		//  (eg. if Oodle is used to unpak ini filed?)
+		if ( GConfig )
+		{
+			GConfig->GetInt(TEXT("OodleDataCompressionFormat"), TEXT("PreallocatedBufferCount"), BufferCount, GEngineIni);
+			if (BufferCount < 0)
+			{
+				BufferCount = 0;
+			}
+		}
+
+		OodleDecoderBufferCount = BufferCount;
+		if (OodleDecoderBufferCount)
+		{
+			OodleDecoders = new OodleDataCompressionDecoder[OodleDecoderBufferCount];
+		}
+	}
+	
+	~OodleDataCompressionDecoders()
+	{
+		// NOTE: currently never freed
+
+		if (OodleDecoderBufferCount)
+		{
+			for (int i = 0; i < OodleDecoderBufferCount; ++i)
+			{
+				if (OodleDecoders[i].OodleDecoderMutex.TryLock())
+				{
+					FMemory::Free(OodleDecoders[i].OodleDecoderMemory);
+					OodleDecoders[i].OodleDecoderMemory = nullptr;
+					OodleDecoders[i].OodleDecoderMutex.Unlock();
+				}
+				else
+				{
+					UE_LOG(OodleDataCompression, Error, TEXT("FOodleDataCompressionFormat - shutting down while in use?"));
+				}
+			}
+
+			delete [] OodleDecoders;
+		}
+	}
+
+	
+	int64 OodleDecode(const void * InCompBuf, int64 InCompBufSize64, void * OutRawBuf, int64 InRawLen64) 
+	{
+		OO_SINTa InCompBufSize = TCheckValueCast<OO_SINTa>(InCompBufSize64);
+		OO_SINTa InRawLen = TCheckValueCast<OO_SINTa>(InRawLen64);
+
+		// find the minimum size needed for this decode, OodleDecoderMemorySize may be larger
+		OodleLZ_Compressor CurCompressor = OodleLZ_GetChunkCompressor(InCompBuf, InCompBufSize, NULL);
+		if( CurCompressor == OodleLZ_Compressor_Invalid )
+		{
+			UE_LOG(OodleDataCompression, Error, TEXT("FOodleDataCompressionFormat - no Oodle compressor found!"));
+			return OODLELZ_FAILED;
+		}
+			
+		OO_SINTa DecoderMemorySize = OodleLZDecoder_MemorySizeNeeded(CurCompressor, InRawLen);
+		void * DecoderMemory = NULL;
+
+		// try to take a mutex for one of the pre-allocated decode buffers
+		for (int i = 0; i < OodleDecoderBufferCount; ++i)
+		{
+			if (OodleDecoders[i].OodleDecoderMutex.TryLock()) 
+			{
+				if (OodleDecoders[i].OodleDecoderMemory == nullptr)
+				{
+					// Haven't allocated yet (we allocate on demand)
+					OodleDecoders[i].OodleDecoderMemory = FMemory::Malloc(OodleDecoderMemorySize);
+					if (OodleDecoders[i].OodleDecoderMemory == nullptr)
+					{
+						UE_LOG(OodleDataCompression, Error, TEXT("FOodleDataCompressionFormat - Failed to allocate preallocated	buffer %d bytes!"), OodleDecoderMemorySize);
+						return OODLELZ_FAILED;
+					}
+				}
+
+				if (OodleDecoders[i].OodleDecoderMemory)
+				{
+					//UE_LOG(OodleDataCompression, Display, TEXT("Decode with lock : %d -> %d"),InCompBufSize,InRawLen );
+
+					OO_SINTa Result = OodleLZ_Decompress(InCompBuf, InCompBufSize, OutRawBuf, InRawLen,
+						OodleLZ_FuzzSafe_Yes, OodleLZ_CheckCRC_Yes, OodleLZ_Verbosity_None,
+						NULL, 0, NULL, NULL,
+						OodleDecoders[i].OodleDecoderMemory, OodleDecoderMemorySize);
+
+					OodleDecoders[i].OodleDecoderMutex.Unlock();
+
+					return (int) Result;
+				}
+
+				OodleDecoders[i].OodleDecoderMutex.Unlock();
+			}
+		}
+			
+		//UE_LOG(OodleDataCompression, Display, TEXT("Decode with malloc : %d -> %d"),InCompBufSize,InRawLen );
+
+		// allocate memory for the decoder so that Oodle doesn't allocate anything internally
+		DecoderMemory = FMemory::Malloc(DecoderMemorySize);
+		if (DecoderMemory == NULL) 
+		{
+			UE_LOG(OodleDataCompression, Error, TEXT("FOodleDataCompressionFormat::OodleDecode - Failed to allocate %d!"), DecoderMemorySize);
+			return OODLELZ_FAILED;
+		}
+
+		OO_SINTa Result = OodleLZ_Decompress(InCompBuf, InCompBufSize, OutRawBuf, InRawLen, 
+			OodleLZ_FuzzSafe_Yes,OodleLZ_CheckCRC_Yes,OodleLZ_Verbosity_None,
+			NULL, 0, NULL, NULL,
+			DecoderMemory, DecoderMemorySize);
+
+		FMemory::Free(DecoderMemory);
+
+		return (int64) Result;
+	}
+
+};
+
+
+
+static OodleLZ_Compressor OodleDataCompressorToOodleLZ_Compressor(EOodleDataCompressor Compressor)
+{
+	switch(Compressor)
+	{
+	case EOodleDataCompressor::Selkie:
+		return OodleLZ_Compressor_Selkie;
+	case EOodleDataCompressor::Mermaid:
+		return OodleLZ_Compressor_Mermaid;
+	case EOodleDataCompressor::Kraken:
+		return OodleLZ_Compressor_Kraken;
+	case EOodleDataCompressor::Leviathan:
+		return OodleLZ_Compressor_Leviathan;
+	case EOodleDataCompressor::NotSet:
+		return OodleLZ_Compressor_Invalid;
+	default:
+		UE_LOG(OodleDataCompression,Error,TEXT("Invalid Compressor: %d\n"),(int)Compressor);
+		return OodleLZ_Compressor_Invalid;
+	}
+}
+
+static OodleLZ_CompressionLevel OodleDataCompressionLevelToOodleLZ_CompressionLevel(EOodleDataCompressionLevel Level)
+{
+	int IntLevel = (int)Level;
+	
+	if ( IntLevel < (int)OodleLZ_CompressionLevel_Min || IntLevel > (int)OodleLZ_CompressionLevel_Max )
+	{
+		UE_LOG(OodleDataCompression,Error,TEXT("Invalid Level: %d\n"),IntLevel);
+		return OodleLZ_CompressionLevel_Invalid;
+	}
+
+	return (OodleLZ_CompressionLevel) IntLevel;
+}
+
+int64 CORE_API OodleDataCompressedBufferSizeNeeded(int64 InUncompressedSize)
+{
+	// size needed is the same for all newlz's
+	//	so don't bother with a compressor arg here
+	return OodleLZ_GetCompressedBufferSizeNeeded(OodleLZ_Compressor_Kraken, TCheckValueCast<OO_SINTa>(InUncompressedSize));
+}
+
+int64 CORE_API OodleDataGetMaximumCompressedSize(int64 InUncompressedSize)
+{
+	int64 NumBlocks = (InUncompressedSize+ OODLELZ_BLOCK_LEN-1)/OODLELZ_BLOCK_LEN;
+	int64 MaxCompressedSize = InUncompressedSize + NumBlocks * OODLELZ_BLOCK_MAXIMUM_EXPANSION;
+	return MaxCompressedSize;
+}
+
+int64 CORE_API OodleDataCompress(
+							void * OutCompressedData, int64 InCompressedBufferSize,
+							const void * InUncompressedData, int64 InUncompressedSize,
+							EOodleDataCompressor Compressor,
+							EOodleDataCompressionLevel Level)
+{
+	OodleLZ_Compressor LZCompressor = OodleDataCompressorToOodleLZ_Compressor(Compressor);
+	OodleLZ_CompressionLevel LZLevel = OodleDataCompressionLevelToOodleLZ_CompressionLevel(Level);
+
+	if ( InCompressedBufferSize < (int64) OodleLZ_GetCompressedBufferSizeNeeded(LZCompressor,TCheckValueCast<OO_SINTa>(InUncompressedSize)) )
+	{
+		UE_LOG(OodleDataCompression,Error,TEXT("OutCompressedSize too small\n"));		
+		return OODLELZ_FAILED;
+	}
+
+	// OodleLZ_Compress will alloc internally using installed CorePlugins
+	//	 (currently default plugins, no plugins installed)
+
+	OO_SINTa CompressedSize = OodleLZ_Compress(LZCompressor,InUncompressedData,InUncompressedSize,OutCompressedData,LZLevel);
+	return (int64) CompressedSize;
+}
+
+static OodleDataCompressionDecoders * GetGlobalOodleDataCompressionDecoders()
+{
+	static OodleDataCompressionDecoders GlobalOodleDataCompressionDecoders;
+	// init on first use, never freed
+	return &GlobalOodleDataCompressionDecoders;
+}
+
+
+static ICompressionFormat * GlobalOodleDataCompressionFormat = nullptr;
+extern ICompressionFormat * CreateOodleDataCompressionFormat();
+
+void CORE_API OodleDataCompressionFormatInitOnFirstUseFromLock()
+{
+	// called from inside a critical section lock 
+	//	from Compression.cpp / GetCompressionFormat
+	if ( GlobalOodleDataCompressionFormat != nullptr )
+		return;
+
+	GlobalOodleDataCompressionFormat = CreateOodleDataCompressionFormat();
+}
+
+
+bool CORE_API OodleDataDecompress(
+						void * OutUncompressedData, int64 InUncompressedSize,
+						const void * InCompressedData, int64 InCompressedSize)
+{
+	OodleDataCompressionDecoders * Decoders = GetGlobalOodleDataCompressionDecoders();
+
+	int64 DecodeSize = Decoders->OodleDecode(InCompressedData,InCompressedSize,OutUncompressedData,InUncompressedSize);
+
+	if ( DecodeSize == OODLELZ_FAILED )
+	{
+		return false;
+	}
+
+	check( DecodeSize == InUncompressedSize );
+	return true;
+}
+					
+void CORE_API OodleDataCompressionStartupPreInit(void)
+{
+	// called from LaunchEngineLoop at "PreInit" time
+	// not all Engine services may be set up yet, be careful what you use
+	
+	// @todo Oodle could install CorePlugins here for log/alloc/etc.
+
+	// OodleConfig set global options for Oodle :
+	OodleConfigValues OodleConfig = { };
+	Oodle_GetConfigValues(&OodleConfig);
+	// UE5 will always read/write Oodle v9 binary data :
+	OodleConfig.m_OodleLZ_BackwardsCompatible_MajorVersion = 9;
+	Oodle_SetConfigValues(&OodleConfig);
+}
diff --git a/Engine/Source/Runtime/Core/Private/Compression/OodleDataCompressionFormat.cpp b/Engine/Source/Runtime/Core/Private/Compression/OodleDataCompressionFormat.cpp
new file mode 100644
--- /dev/null
+++ b/Engine/Source/Runtime/Core/Private/Compression/OodleDataCompressionFormat.cpp
@@ -0,0 +1,445 @@
+// Copyright Epic Games, Inc. All Rights Reserved.
+
+
+
+// OodleDataCompressionFormat :
+//	Oodle in Unreal for pak/iostore
+//	sets compression options from command line
+//	(automation tool turns config into command line for us)
+//  installs the "Oodle" named CompressionFormat
+
+
+/*****************
+* 
+* Oodle Data compression format
+* provides Oodle compression for Pak files & iostore
+* is not for generic compression usage
+* 
+* The Oodle LZ codecs are extremely fast to decode and almost always speed up load times
+* 
+* The codecs are :
+* Kraken  : high compression with good decode speed, the usual default
+* Mermaid : less compression and faster decode speed; good when CPU bound or on platforms with less CPU power
+* Selkie  : even less compression and faster that Mermaid
+* Leviathan : more compression and slower to decode than Kraken
+* 
+* The encode time is mostly independent of the codec.  Use the codec to choose decode speed, and the encoder effort
+* level to control encode time.
+* 
+* For daily iteration you might want level 3 ("Fast").  For shipping packages you might want level 6 ("optimal2") or higher.
+* The valid level range is -4 to 9
+* 
+* This plugin reads its options on the command line via "compressmethod" and "compresslevel"
+* eg. "-compressmethod=Kraken -compresslevel=4"
+* 
+* The Oodle decoder can decode any codec used, it doesn't need to know which one was used.
+* 
+* Compression options should be set up in your Game.ini ; for example :
+* 
+
+[/Script/UnrealEd.ProjectPackagingSettings]
+bCompressed=True
+bForceUseProjectCompressionFormat=False
+PackageCompressionFormat=Oodle
+PackageAdditionalCompressionOptions=-compressionblocksize=256KB
+PackageCompressionMethod=Kraken
+PackageCompressionLevel_Distribution=7
+PackageCompressionLevel_TestShipping=5
+PackageCompressionLevel_DebugDevelopment=4
+
+* This can be set in DefaultGame.ini then overrides set up per-platform.
+* 
+* The Engine also has a veto compressionformat set up in the DataDrivenPlatformInfo.ini for each platform in the field
+* "HardwareCompressionFormat"
+* eg. platforms that don't want any software compressor can set "HardwareCompressionFormat=None" and this will override what you
+* set in "PackageCompressionFormat".
+* 
+* The idea is in typical use, you set "PackageCompressionFormat" for your Game, and you get that compressor on most platforms, but on
+* some platforms that don't want compression, it automatically turns off.
+* 
+* If you want to force use of your Game.ini compressor (ignore the HardwareCompressionFormat) you can set bForceUseProjectCompressionFormat
+* in ProjectPackagingSettings.
+* 
+* 
+* ***************************/
+
+#include "CoreMinimal.h"
+#include "CoreTypes.h"
+#include "Misc/Compression.h"
+#include "Misc/ICompressionFormat.h"
+#include "Misc/CommandLine.h"
+#include "Misc/ConfigCacheIni.h"
+#include "Misc/Parse.h"
+
+/*
+#if WITH_EDITOR
+#include "Settings/ProjectPackagingSettings.h"
+#endif
+*/
+
+#include "HAL/PlatformProcess.h"
+
+#include "Compression/OodleDataCompression.h"
+#include "oodle2.h"
+
+
+#define OODLE_DERIVEDDATA_VER TEXT("BA7AA26CD1C3498787A3F3AA53895042")
+
+// function pointer for old DLL import :
+extern "C"
+{
+typedef OO_SINTa (OOLINK t_fp_OodleLZ_Compress)(OodleLZ_Compressor compressor,
+    const void * rawBuf,OO_SINTa rawLen,void * compBuf,
+    OodleLZ_CompressionLevel selection,
+    const OodleLZ_CompressOptions * pOptions,
+    const void * dictionaryBase,
+    const void * lrm);
+};
+
+
+struct FOodleDataCompressionFormat : ICompressionFormat
+{
+	OodleLZ_Compressor Compressor;
+	OodleLZ_CompressionLevel CompressionLevel;
+	OodleLZ_CompressOptions CompressionOptions;
+	
+	//  OodleLZCompressFuncPtr is non-null of a DLL/so was loaded for encoding
+	t_fp_OodleLZ_Compress  * OodleLZCompressFuncPtr = nullptr;
+
+
+	FOodleDataCompressionFormat(OodleLZ_Compressor InCompressor, OodleLZ_CompressionLevel InCompressionLevel, int InSpaceSpeedTradeoffBytes, void * InOodleCompressFuncPtr)
+	{
+		Compressor = InCompressor;
+		CompressionLevel = InCompressionLevel;
+		CompressionOptions = *OodleLZ_CompressOptions_GetDefault(Compressor, CompressionLevel);
+		CompressionOptions.spaceSpeedTradeoffBytes = InSpaceSpeedTradeoffBytes;
+		// we're usually doing limited chunks, no need for LRM :
+		CompressionOptions.makeLongRangeMatcher = false;
+		
+		OodleLZCompressFuncPtr = (t_fp_OodleLZ_Compress *) InOodleCompressFuncPtr;
+	}
+
+	virtual ~FOodleDataCompressionFormat() 
+	{
+	}
+
+	FString GetCompressorString() const
+	{
+		// convert values to enums
+		switch (Compressor)
+		{
+			case OodleLZ_Compressor_Selkie:				return TEXT("Selkie");
+			case OodleLZ_Compressor_Mermaid:			return TEXT("Mermaid");
+			case OodleLZ_Compressor_Kraken: 			return TEXT("Kraken");
+			case OodleLZ_Compressor_Leviathan:			return TEXT("Leviathan");
+			case OodleLZ_Compressor_Hydra:				return TEXT("Hydra");
+			default: break;
+		}
+		return TEXT("Unknown");
+	}
+
+	FString GetCompressionLevelString() const
+	{
+		switch (CompressionLevel)
+		{
+			case OodleLZ_CompressionLevel_HyperFast4:	return TEXT("HyperFast4");
+			case OodleLZ_CompressionLevel_HyperFast3:	return TEXT("HyperFast3");
+			case OodleLZ_CompressionLevel_HyperFast2:	return TEXT("HyperFast2");
+			case OodleLZ_CompressionLevel_HyperFast1:	return TEXT("HyperFast1");
+			case OodleLZ_CompressionLevel_None:			return TEXT("None");
+			case OodleLZ_CompressionLevel_SuperFast:	return TEXT("SuperFast");
+			case OodleLZ_CompressionLevel_VeryFast:		return TEXT("VeryFast");
+			case OodleLZ_CompressionLevel_Fast:			return TEXT("Fast");
+			case OodleLZ_CompressionLevel_Normal:		return TEXT("Normal");
+			case OodleLZ_CompressionLevel_Optimal1:		return TEXT("Optimal1");
+			case OodleLZ_CompressionLevel_Optimal2:		return TEXT("Optimal2");
+			case OodleLZ_CompressionLevel_Optimal3:		return TEXT("Optimal3");
+			case OodleLZ_CompressionLevel_Optimal4:		return TEXT("Optimal4");
+			case OodleLZ_CompressionLevel_Optimal5:		return TEXT("Optimal5");
+			default: break;
+		}
+		return TEXT("Unknown");
+	}
+	
+	virtual bool DoesOwnWorthDecompressingCheck() override
+	{
+		// Oodle does own "worth it" check internally, don't add one
+		return true;
+	}
+
+	virtual FName GetCompressionFormatName() override
+	{
+		return NAME_Oodle;
+	}
+
+	virtual uint32 GetVersion() override
+	{
+		return 20000 + OODLE2_VERSION_MAJOR*100 + OODLE2_VERSION_MINOR;
+	}
+	
+
+	virtual FString GetDDCKeySuffix() override
+	{
+		// DerivedDataCache key string
+		//  ideally this should be unique for any settings changed
+		return FString::Printf(TEXT("C_%s_CL_%s_%s"), *GetCompressorString(), *GetCompressionLevelString(), OODLE_DERIVEDDATA_VER);
+	}
+
+	virtual bool Compress(void* OutCompressedBuffer, int32& OutCompressedSize, const void* InUncompressedBuffer, int32 InUncompressedSize, int32 InCompressionData) override
+	{
+		// OutCompressedSize is read-write
+		int32 CompressedBufferSize = OutCompressedSize;
+
+		// CompressedSize should be >= GetCompressedBufferSize(UncompressedSize, CompressionData)
+		check(CompressedBufferSize >= GetCompressedBufferSize(InUncompressedSize, InCompressionData));
+		
+		OO_SINTa Result = 0;
+
+		if ( OodleLZCompressFuncPtr != nullptr )
+		{
+			// encode using func pointer to loaded DLL :
+			Result = (*OodleLZCompressFuncPtr)(Compressor, InUncompressedBuffer, InUncompressedSize, OutCompressedBuffer, CompressionLevel, NULL, NULL, NULL);
+
+			// NOTE: not using &CompressOptions
+			//	because OodleLZ_CompressOptions_GetDefault has changed from v5
+			//	need to either call that in the DLL func ptr or not use it
+			//	(SSTB == 0 doesn't mean default in old Oodle)
+
+			#if 0 // DO_CHECK
+			if ( Result > 0 )
+			{
+				// verify we can decode data made with DLL-enocder using the current decoder :
+
+				int32 DecodeSpaceSize = InUncompressedSize + 16; // FuzzSafe_Yes so shouldn't need big padding even on Oodle v5
+				void * DecodeSpace = FMemory::Malloc(DecodeSpaceSize);
+				check( DecodeSpace != nullptr );
+
+				OO_SINTa DecRet = OodleLZ_Decompress(OutCompressedBuffer,Result,DecodeSpace,InUncompressedSize,OodleLZ_FuzzSafe_Yes);
+				check( DecRet == InUncompressedSize );
+
+				int cmp = memcmp(DecodeSpace,InUncompressedBuffer,InUncompressedSize);
+				check( cmp == 0 );
+
+				FMemory::Free(DecodeSpace);
+			}
+			#endif
+		}
+		else
+		{
+			Result = OodleLZ_Compress(Compressor, InUncompressedBuffer, InUncompressedSize, OutCompressedBuffer, CompressionLevel, &CompressionOptions);
+		}
+
+		// verbose log all compresses :
+		//UE_LOG(OodleDataCompression, Display, TEXT("Oodle Compress : %d -> %d"), UncompressedSize, Result);
+		
+		if (Result <= 0)
+		{
+			OutCompressedSize = -1;
+			return false;
+		}
+		else
+		{
+			OutCompressedSize = (int32) Result;
+			return true;
+		}
+	}
+
+	virtual bool Uncompress(void* OutUncompressedBuffer, int32& OutUncompressedSize, const void* InCompressedBuffer, int32 InCompressedSize, int32 CompressionData) override
+	{
+		// OutUncompressedSize is read-write
+		return OodleDataDecompress(
+			OutUncompressedBuffer, OutUncompressedSize,
+			InCompressedBuffer,InCompressedSize);
+	}
+
+	virtual int32 GetCompressedBufferSize(int32 UncompressedSize, int32 CompressionData) override
+	{
+		// CompressionData is not used
+		int32 Needed = (int32)OodleLZ_GetCompressedBufferSizeNeeded(Compressor, UncompressedSize);
+
+		if ( OodleLZCompressFuncPtr != nullptr )
+		{
+			// older versions of Oodle needed larger padding
+			// old rule was 274 per 256KB :
+
+			int32 NumSeekChunks = (UncompressedSize + OODLELZ_BLOCK_LEN-1)/OODLELZ_BLOCK_LEN;
+
+			Needed += 274*NumSeekChunks;
+
+			// can we call OodleLZ_GetCompressedBufferSizeNeeded from the DLL func ptr?
+			//  OodleLZ_GetCompressedBufferSizeNeeded changed signature at one point to add Compressor in 2.8.0
+			//	so you'd need to know what version to use
+		}
+
+		return Needed;
+	}
+};
+
+extern ICompressionFormat * CreateOodleDataCompressionFormat();
+
+ICompressionFormat * CreateOodleDataCompressionFormat()
+{
+	// set up Oodle for packaging by parsing the command line options
+	// this is now down on first use of a pluggable ICompressionFormat
+
+
+	// settings to use in non-tools context (eg. runtime game encoding) :
+	// (SetDefaultOodleOptionsForPackaging sets options for pak compression & iostore)
+	OodleLZ_Compressor UsedCompressor = OodleLZ_Compressor_Mermaid;
+	OodleLZ_CompressionLevel UsedLevel = OodleLZ_CompressionLevel_VeryFast;
+	int32 SpaceSpeedTradeoff = OODLELZ_SPACESPEEDTRADEOFFBYTES_DEFAULT;
+	void * OodleCompressFuncPtr = nullptr;
+
+	#if ( (!UE_BUILD_SHIPPING) || WITH_EDITOR )
+	{
+
+	// parse the command line to get compressor & level settings
+
+	// this Startup is done in various different contexts;
+	// when the editor loads up
+	// when the game loads (we will be used to decompress only and encode settings are not used)
+	// when the package cooking tool loads up <- this is when we set the relevant encode settings
+		
+	// beware init order; instead of calling IsRunningCommandlet it's safer to just look for -run=
+
+	// is_program is true for cooker & UnrealPak (not Editor or Game)
+	bool IsProgram = IS_PROGRAM;
+	bool IsIOStore = FCString::Strifind(FCommandLine::Get(), TEXT("-run=iostore")) != NULL;
+				
+	// we only need to be doing all this when run as UnrealPak or iostore commandlet
+	//	(IsProgram also picks up cooker and a few other things, that's okay)
+	if ( IsProgram || IsIOStore )
+	{
+		// defaults if no options set :
+		UsedCompressor = OodleLZ_Compressor_Kraken;
+		// Kraken is a good compromise of compression ratio & speed
+
+		UsedLevel = OodleLZ_CompressionLevel_Normal;
+		// for faster iteration time during development
+
+		SpaceSpeedTradeoff = OODLELZ_SPACESPEEDTRADEOFFBYTES_DEFAULT;
+		// SpaceSpeedTradeoff is mainly for tuning the Hydra compressor
+		//	it can also be used to skew your compression towards higher ratio vs faster decode
+
+		// convert values to enums
+		TMap<FString, OodleLZ_Compressor> MethodMap =
+		{ 
+			{ TEXT("Selkie"), OodleLZ_Compressor_Selkie },
+			{ TEXT("Mermaid"), OodleLZ_Compressor_Mermaid },
+			{ TEXT("Kraken"), OodleLZ_Compressor_Kraken },
+			{ TEXT("Leviathan"), OodleLZ_Compressor_Leviathan },
+			{ TEXT("Hydra"), OodleLZ_Compressor_Hydra },
+			// when adding here remember to update FOodleDataCompressionFormat::GetCompressorString()
+		};
+		TMap<FString, OodleLZ_CompressionLevel> LevelMap = 
+		{ 
+			{ TEXT("HyperFast4"), OodleLZ_CompressionLevel_HyperFast4 },
+			{ TEXT("HyperFast3"), OodleLZ_CompressionLevel_HyperFast3 },
+			{ TEXT("HyperFast2"), OodleLZ_CompressionLevel_HyperFast2 },
+			{ TEXT("HyperFast1"), OodleLZ_CompressionLevel_HyperFast1 },
+			{ TEXT("HyperFast"), OodleLZ_CompressionLevel_HyperFast1 },
+			{ TEXT("None")     , OodleLZ_CompressionLevel_None },
+			{ TEXT("SuperFast"), OodleLZ_CompressionLevel_SuperFast },
+			{ TEXT("VeryFast"), OodleLZ_CompressionLevel_VeryFast },
+			{ TEXT("Fast")  , OodleLZ_CompressionLevel_Fast },
+			{ TEXT("Normal"), OodleLZ_CompressionLevel_Normal },
+			{ TEXT("Optimal1"), OodleLZ_CompressionLevel_Optimal1 },
+			{ TEXT("Optimal2"), OodleLZ_CompressionLevel_Optimal2 },
+			{ TEXT("Optimal") , OodleLZ_CompressionLevel_Optimal2 },
+			{ TEXT("Optimal3"), OodleLZ_CompressionLevel_Optimal3 },
+			{ TEXT("Optimal4"), OodleLZ_CompressionLevel_Optimal4 },
+			{ TEXT("Optimal5"), OodleLZ_CompressionLevel_Optimal5 },
+				
+			{ TEXT("-4"), OodleLZ_CompressionLevel_HyperFast4 },
+			{ TEXT("-3"), OodleLZ_CompressionLevel_HyperFast3 },
+			{ TEXT("-2"), OodleLZ_CompressionLevel_HyperFast2 },
+			{ TEXT("-1"), OodleLZ_CompressionLevel_HyperFast1 },
+			{ TEXT("0"), OodleLZ_CompressionLevel_None },
+			{ TEXT("1"), OodleLZ_CompressionLevel_SuperFast },
+			{ TEXT("2"), OodleLZ_CompressionLevel_VeryFast },
+			{ TEXT("3"), OodleLZ_CompressionLevel_Fast },
+			{ TEXT("4"), OodleLZ_CompressionLevel_Normal },
+			{ TEXT("5"), OodleLZ_CompressionLevel_Optimal1 },
+			{ TEXT("6"), OodleLZ_CompressionLevel_Optimal2 },
+			{ TEXT("7"), OodleLZ_CompressionLevel_Optimal3 },
+			{ TEXT("8"), OodleLZ_CompressionLevel_Optimal4 },
+			{ TEXT("9"), OodleLZ_CompressionLevel_Optimal5 },
+			// when adding here remember to update FOodleDataCompressionFormat::GetCompressionLevelString()
+		};
+
+		// override from command line :
+		FString Method = "";
+		FString Level = "";
+
+		// let commandline override
+		//FParse::Value does not change output fields if they are not found
+		FParse::Value(FCommandLine::Get(), TEXT("compressmethod="), Method);
+		FParse::Value(FCommandLine::Get(), TEXT("compresslevel="), Level);
+		FParse::Value(FCommandLine::Get(), TEXT("OodleSpaceSpeedTradeoff="), SpaceSpeedTradeoff);
+		
+		OodleLZ_Compressor * pUsedCompressor = MethodMap.Find(Method);
+		if (pUsedCompressor)
+		{
+			UsedCompressor = *pUsedCompressor;
+		}
+
+		OodleLZ_CompressionLevel * pUsedLevel = LevelMap.Find(Level);
+		if (pUsedLevel)
+		{
+			UsedLevel = *pUsedLevel;
+		}
+
+		// no init log line if we're not enabled for pak/iostore :
+		bool bUseCompressionFormatOodle = FCString::Strifind(FCommandLine::Get(), TEXT("-compressionformats=oodle")) != NULL;
+		if ( bUseCompressionFormatOodle )			
+		{
+			UE_LOG(OodleDataCompression, Display, TEXT("Oodle v%s format for pak/iostore with method=%s, level=%d=%s"), TEXT(OodleVersion), **MethodMap.FindKey(UsedCompressor), (int)UsedLevel, **LevelMap.FindKey(UsedLevel) );
+		}
+		
+		// OodleCompressDLL to use an earlier *encoder* to maintain identical bit streams
+		// mainly for shipped games to avoid patches
+		// do not specify this if you want to use the latest Engine version of Oodle for encoding
+		FString OodleDLL = "";
+		if ( FParse::Value(FCommandLine::Get(), TEXT("OodleCompressDLL="), OodleDLL) && ! OodleDLL.IsEmpty() )
+		{
+			UE_LOG(OodleDataCompression, Display, TEXT("OodleCompressDLL=%s"), *OodleDLL);
+				
+			// just load from default locations :
+			//  the DLL you want here should be added in RuntimeDependencies.Add in a build.cs
+			//  to ensure it is in the binaries directory of the process
+			void * 	OodleDLLHandle = FPlatformProcess::GetDllHandle(*OodleDLL);
+
+			if ( OodleDLLHandle == nullptr )
+			{
+				UE_LOG(OodleDataCompression, Warning, TEXT("OodleCompressDLL %s requested but could not be loaded"), *OodleDLL);
+			}
+			else
+			{
+				// GetDllExport is GetProcAddress or dlsym on Mac
+				//  LinuxPlatformProcess::GetDllExport seems to be missing ?
+				//OodleCompressFuncPtr = GetProcAddress( (HMODULE)OodleDLLHandle, "OodleLZ_Compress" );
+				OodleCompressFuncPtr = FPlatformProcess::GetDllExport( OodleDLLHandle, TEXT("OodleLZ_Compress") );
+				if ( OodleCompressFuncPtr == nullptr )
+				{
+					UE_LOG(OodleDataCompression, Warning, TEXT("OodleCompressDLL %s loaded but no OodleLZ_Compress !?"), *OodleDLL);
+				}
+			}
+		}
+	}
+
+	}
+	#endif // SHIPPING
+
+	//-----------------------------------
+	// register the compression format :
+	//  this is used by the shipping game to decode any paks compressed with Oodle :
+
+	ICompressionFormat * CompressionFormat = new FOodleDataCompressionFormat(UsedCompressor, UsedLevel, SpaceSpeedTradeoff, OodleCompressFuncPtr);
+
+	IModularFeatures::Get().RegisterModularFeature(COMPRESSION_FORMAT_FEATURE_NAME, CompressionFormat);
+
+	// NOTE: currently never shutdown and unregistered
+
+	return CompressionFormat;
+}
+
diff --git a/Engine/Source/Runtime/Core/Private/Compression/OodleDataCompressionUtil.cpp b/Engine/Source/Runtime/Core/Private/Compression/OodleDataCompressionUtil.cpp
new file mode 100644
--- /dev/null
+++ b/Engine/Source/Runtime/Core/Private/Compression/OodleDataCompressionUtil.cpp
@@ -0,0 +1,62 @@
+// Copyright Epic Games, Inc. All Rights Reserved.
+
+#include "Compression/OodleDataCompressionUtil.h"
+#include "Templates/CheckValueCast.h"
+
+#include "Serialization/MemoryReader.h"
+#include "Serialization/MemoryWriter.h"
+
+namespace OodleDataCompressionUtil
+{
+	bool CORE_API DecompressReplayData(const TArray<uint8>& InCompressed, TArray< uint8 >& OutBuffer)
+	{
+		int32 Size = 0;
+		int32 CompressedSize = 0;
+
+		FMemoryReader Reader(InCompressed);
+		Reader << Size;
+		Reader << CompressedSize;
+
+		OutBuffer.SetNum(Size, false);
+
+		if ( Reader.Tell() + CompressedSize > InCompressed.Num() )
+			return false;
+
+		return OodleDataDecompress(OutBuffer.GetData(), OutBuffer.Num(), InCompressed.GetData() + Reader.Tell(), CompressedSize );
+	}
+
+	bool CORE_API CompressReplayData(const TArray<uint8>& InBuffer, TArray< uint8 >& OutCompressed)
+	{
+		const void * InPtr = InBuffer.GetData();
+		const int32 InSize = InBuffer.Num();
+		int32 CompressedSize = 0;
+
+		int32 ReservedBytes = sizeof(int32)*2;
+				
+		const EOodleDataCompressor			Compressor = EOodleDataCompressor::Selkie;
+		const EOodleDataCompressionLevel	Level = EOodleDataCompressionLevel::VeryFast;
+		
+		int32 CompSizeNeeded = (int32) OodleDataCompressedBufferSizeNeeded(InSize);
+		OutCompressed.SetNum( TCheckValueCast<int32>( ReservedBytes + CompSizeNeeded) );
+		void * CompPtr = OutCompressed.GetData() + ReservedBytes;
+
+		CompressedSize = (int32) OodleDataCompress(CompPtr,CompSizeNeeded,InPtr,InSize,Compressor,Level);
+		if ( CompressedSize <= 0 )
+		{
+			return false;
+		}
+		
+		int32 InSize32 = TCheckValueCast<int32>(InSize);
+		int32 CompressedSize32 = TCheckValueCast<int32>(CompressedSize);
+
+		FMemoryWriter Writer(OutCompressed);
+		Writer << InSize32;
+		Writer << CompressedSize32;
+		
+		check( Writer.Tell() == ReservedBytes );
+
+		OutCompressed.SetNum( TCheckValueCast<int32>(CompressedSize + ReservedBytes) , false);
+
+		return true;
+	}
+};
diff --git a/Engine/Source/Runtime/Core/Private/Misc/Compression.cpp b/Engine/Source/Runtime/Core/Private/Misc/Compression.cpp
--- a/Engine/Source/Runtime/Core/Private/Misc/Compression.cpp
+++ b/Engine/Source/Runtime/Core/Private/Misc/Compression.cpp
@@ -13,6 +13,9 @@
 #include "Misc/ICompressionFormat.h"
 
 #include "Misc/MemoryReadStream.h"
+
+#include "Compression/OodleDataCompression.h"
+
 // #include "TargetPlatformBase.h"
 THIRD_PARTY_INCLUDES_START
 #include "ThirdParty/zlib/zlib-1.2.5/Inc/zlib.h"
@@ -30,6 +33,7 @@ DECLARE_STATS_GROUP( TEXT( "Compression" ), STATGROUP_Compression, STATCAT_Advan
 
 PRAGMA_DISABLE_UNSAFE_TYPECAST_WARNINGS
 
+// static class members for registered plugin/module compression formats :
 TMap<FName, struct ICompressionFormat*> FCompression::CompressionFormats;
 FCriticalSection FCompression::CompressionFormatsCriticalSection;
 
@@ -438,9 +442,19 @@ uint32 FCompression::GetCompressorVersion(FName FormatName)
 ICompressionFormat* FCompression::GetCompressionFormat(FName FormatName, bool bErrorOnFailure)
 {
 	FScopeLock Lock(&CompressionFormatsCriticalSection);
+
 	ICompressionFormat** ExistingFormat = CompressionFormats.Find(FormatName);
 	if (ExistingFormat == nullptr)
 	{
+		if ( FormatName == NAME_Oodle )
+		{
+			// Oodle ICompressionFormat is created on first use :
+			// inside CompressionFormatsCriticalSection lock
+			OodleDataCompressionFormatInitOnFirstUseFromLock();
+
+			// OodleDataCompressionFormatInitOnFirstUseFromLock added it to the ModularFeatures list
+		}
+
 		TArray<ICompressionFormat*> Features = IModularFeatures::Get().GetModularFeatureImplementations<ICompressionFormat>(COMPRESSION_FORMAT_FEATURE_NAME);
 
 		for (ICompressionFormat* CompressionFormat : Features)
@@ -482,7 +496,9 @@ FName FCompression::GetCompressionFormatFromDeprecatedFlags(ECompressionFlags Fl
 			return NAME_Gzip;
 		// COMPRESS_Custom was a temporary solution to third party compression before we had plugins working, and it was only ever used with oodle, we just assume Oodle with Custom
 		case COMPRESS_Custom:
-			return TEXT("Oodle");
+			return NAME_Oodle;
+		default:
+			break;
 	}
 
 	return NAME_None;
@@ -490,7 +506,6 @@ FName FCompression::GetCompressionFormatFromDeprecatedFlags(ECompressionFlags Fl
 
 int32 FCompression::CompressMemoryBound(FName FormatName, int32 UncompressedSize, ECompressionFlags Flags, int32 CompressionData)
 {
-	TRACE_CPUPROFILER_EVENT_SCOPE(FCompression::CompressMemoryBound);
 	int32 CompressionBound = UncompressedSize;
 
 	if (FormatName == NAME_Zlib)
@@ -541,23 +556,19 @@ bool FCompression::CompressMemoryIfWorthDecompressing(FName FormatName, int32 Mi
 		// if input size is smaller than the number of bytes we need to save
 		// no need to even try encoding
 		// also saves encode time
+		// NOTE : this check applies even for compressor who say "bNeedsWorthItCheck = false" , eg. Oodle
 		return false;
 	}
 	
 	bool bNeedsWorthItCheck = false;
-	if (FormatName == NAME_Zlib)
-	{
-		// hardcoded zlib
-		bNeedsWorthItCheck = true;
-	}
-	else if (FormatName == NAME_Gzip)
+	
+	if (FormatName == NAME_Oodle)
 	{
-		// hardcoded gzip
-		bNeedsWorthItCheck = true;
+		// Oodle does its own internal "worth it" check
+		bNeedsWorthItCheck = false;
 	}
-	else if (FormatName == NAME_LZ4)
+	else if (FormatName == NAME_Zlib || FormatName == NAME_Gzip || FormatName == NAME_LZ4)
 	{
-		// hardcoded lz4
 		bNeedsWorthItCheck = true;
 	}
 	else
@@ -707,6 +718,12 @@ bool FCompression::UncompressMemory(FName FormatName, void* UncompressedBuffer,
 		// hardcoded lz4
 		bUncompressSucceeded = LZ4_decompress_safe((const char*)CompressedBuffer, (char*)UncompressedBuffer, CompressedSize, UncompressedSize) > 0;
 	}
+	else if (FormatName == NAME_Oodle)
+	{
+		// hardcoded Oodle
+		// can decode Oodle data without creating Oodle ICompressionFormat
+		bUncompressSucceeded = OodleDataDecompress(UncompressedBuffer,UncompressedSize,CompressedBuffer,CompressedSize);
+	}
 	else
 	{
 		// let the format module compress it
@@ -727,7 +744,11 @@ bool FCompression::UncompressMemory(FName FormatName, void* UncompressedBuffer,
 			FFailOnUncompressErrors()
 				: Value(true) // fail by default
 			{
-				GConfig->GetBool(TEXT("Core.System"), TEXT("FailOnUncompressErrors"), Value, GEngineIni);
+				// very early decodes of first paks could be before this config is loaded
+				if ( GConfig )
+				{
+					GConfig->GetBool(TEXT("Core.System"), TEXT("FailOnUncompressErrors"), Value, GEngineIni);
+				}
 			}
 		} FailOnUncompressErrors;
 		if (!FailOnUncompressErrors.Value)
@@ -963,7 +984,7 @@ void* FCompressedGrowableBuffer::Access( int32 Offset )
 bool FCompression::IsFormatValid(FName FormatName)
 {
 	// build in formats are always valid
-	if (FormatName == NAME_Zlib || FormatName == NAME_Gzip)
+	if (FormatName == NAME_Zlib || FormatName == NAME_Gzip || FormatName == NAME_LZ4 || FormatName == NAME_Oodle)
 	{
 		return true;
 	}
diff --git a/Engine/Source/Runtime/Core/Public/Compression/OodleDataCompression.h b/Engine/Source/Runtime/Core/Public/Compression/OodleDataCompression.h
new file mode 100644
--- /dev/null
+++ b/Engine/Source/Runtime/Core/Public/Compression/OodleDataCompression.h
@@ -0,0 +1,125 @@
+// Copyright Epic Games, Inc. All Rights Reserved.
+
+#pragma once
+
+#include "CoreTypes.h"
+#include "Logging/LogCategory.h"
+#include "Logging/LogMacros.h"
+
+// OodleDataCompression :
+//	Unreal API for direct access to Oodle Data lossless data compression
+//	for manual encoding (not in Pak/iostore)
+//	NOTE : for any data that will be stored to disk, you should not be compressing it yourself!
+//		allow the pak/iostore system to choose the compressor!
+//
+// Let me emphasize again : for data to be stored in packages for shipping games, do NOT use this
+//	let the pak/iostore system choose the compression from the config options and platform settings
+//
+// This is for utility compression in non-shipping game package scenarios
+//	eg. uassets, storage caches and back-ends, large network transfers
+
+DECLARE_LOG_CATEGORY_EXTERN(OodleDataCompression, Log, All);
+
+/**
+ * EOodleDataCompressor : Choose the Oodle Compressor
+ *  this mostly trades decompression speed vs compression ratio
+ * 
+ * From fastest to slowest (to decode) : Selkie, Mermaid, Kraken, Leviathan
+ * 
+ * encode speed is determined by EOodleDataCompressionLevel , not the compressor choice.
+ * 
+ * When in doubt, start with Kraken
+ * Representative compression ratios and decode speeds :
+ 
+Selkie4       :  1.86:1 ,   4232.6 dec MB/s
+Mermaid4      :  2.21:1 ,   2648.9 dec MB/s
+Kraken4       :  2.44:1 ,   1467.1 dec MB/s
+Leviathan4    :  2.62:1 ,    961.8 dec MB/s
+ * 
+ */
+//UENUM() // @todo Oodle might be nice if these were UENUM but can't pull UObject from inside Core?
+// enum values should not change, they may be persisted
+enum class EOodleDataCompressor : uint8
+{
+	NotSet = 0,
+	Selkie = 1,
+	Mermaid = 2,
+	Kraken  = 3,
+	Leviathan = 4
+};
+
+
+/**
+ * EOodleDataCompressionLevel : Choose the Oodle Compression Level
+ *  this mostly trades encode speed vs compression ratio
+ * 
+ * If in doubt start with "Normal" (level 4) then move up or down from there
+ * 
+ * The standard range is in "SuperFast" - "Normal" (levels 1-4)
+ * 
+ * HyperFast is for real-time encoding with minimal compression
+ * 
+ * The "Optimal" levels are much slower to encode, but provide more compression
+ * they are intended for offline cooks
+ * 
+ * representative encode speeds with compression ratio trade off :
+ 
+Kraken-4      :  1.55:1 ,  718.3 enc MB/s
+Kraken-3      :  1.71:1 ,  541.8 enc MB/s
+Kraken-2      :  1.88:1 ,  434.0 enc MB/s
+Kraken-1      :  2.10:1 ,  369.1 enc MB/s
+Kraken1       :  2.27:1 ,  242.6 enc MB/s
+Kraken2       :  2.32:1 ,  157.4 enc MB/s
+Kraken3       :  2.39:1 ,   34.8 enc MB/s
+Kraken4       :  2.44:1 ,   22.4 enc MB/s
+Kraken5       :  2.55:1 ,   10.1 enc MB/s
+Kraken6       :  2.56:1 ,    5.4 enc MB/s
+Kraken7       :  2.64:1 ,    3.7 enc MB/s
+
+ */
+//UENUM() // @todo Oodle might be nice if these were UENUM but can't pull UObject from inside Core?
+// EOodleDataCompressionLevel must numerically match the Oodle internal enum values
+enum class EOodleDataCompressionLevel : int8
+{
+	HyperFast4 = -4,
+	HyperFast3 = -3,
+	HyperFast2 = -2,
+	HyperFast1 = -1,
+	None = 0,
+	SuperFast = 1,
+	VeryFast = 2,
+	Fast = 3,
+	Normal = 4,
+	Optimal1 = 5,
+	Optimal2 = 6,
+	Optimal3 = 7,
+	Optimal4 = 8,
+};
+
+int64 CORE_API OodleDataCompressedBufferSizeNeeded(int64 IncompressedSize);
+
+int64 CORE_API OodleDataGetMaximumCompressedSize(int64 UncompressedSize);
+
+// OutCompressedSize must be >= OodleDataCompressedBufferSizeNeeded(InUncompressedSize)
+// returns compressed size or zero for failure
+int64 CORE_API OodleDataCompress(
+							void * OutCompressedData, int64 CompressedBufferSize,
+							const void * InUncompressedData, int64 UncompressedSize,
+							EOodleDataCompressor Compressor,
+							EOodleDataCompressionLevel Level);
+
+// OutUncompressedSize is the number of bytes to decompress, it must match the number that you encoded
+// InCompressedSize is the buffer size of compressed data, it must be >= the number of compressed bytes needed
+bool CORE_API OodleDataDecompress(
+						void * OutUncompressedData, int64 UncompressedSize,
+						const void * InCompressedData, int64 CompressedSize
+						);
+
+
+// For Compression to/from TArray and such higher level actions use OodleDataCompressionUtil.h
+
+// from Compression.cpp :
+void CORE_API OodleDataCompressionFormatInitOnFirstUseFromLock();
+
+// from LaunchEngineLoop :
+void CORE_API OodleDataCompressionStartupPreInit();
diff --git a/Engine/Source/Runtime/Core/Public/Compression/OodleDataCompressionUtil.h b/Engine/Source/Runtime/Core/Public/Compression/OodleDataCompressionUtil.h
new file mode 100644
--- /dev/null
+++ b/Engine/Source/Runtime/Core/Public/Compression/OodleDataCompressionUtil.h
@@ -0,0 +1,32 @@
+// Copyright Epic Games, Inc. All Rights Reserved.
+
+#pragma once
+
+#include "OodleDataCompression.h"
+
+// OodleDataCompressionUtil : utilities for common Unreal actions built on top of OodleDataCompression
+
+namespace OodleDataCompressionUtil
+{
+	/**
+	* Decompress replay data using Oodle, for use by INetworkReplayStreamer streamers
+	*
+	* @param InCompressed	The compressed replay source data (prefixed with size and uncompressed size)
+	* @param OutBuffer		The destination buffer for uncompressed data
+	* @return				Whether or not decompression succeeded
+	*/
+	bool CORE_API DecompressReplayData(const TArray<uint8>& InCompressed, TArray< uint8 >& OutBuffer);
+
+	/**
+	* Compress replay data using Oodle, for use by INetworkReplayStreamer streamers
+	*
+	* @param InBuffer		The uncompressed replay source data 
+	* @param OutCompressed	The destination buffer for compressed data (prefixed with size and uncompressed size)
+	* @return				Whether or not compression succeeded
+	*/
+	bool CORE_API CompressReplayData(const TArray<uint8>& InBuffer, TArray< uint8 >& OutCompressed);
+};
+
+
+
+
diff --git a/Engine/Source/Runtime/Core/Public/UObject/UnrealNames.inl b/Engine/Source/Runtime/Core/Public/UObject/UnrealNames.inl
--- a/Engine/Source/Runtime/Core/Public/UObject/UnrealNames.inl
+++ b/Engine/Source/Runtime/Core/Public/UObject/UnrealNames.inl
@@ -153,6 +153,7 @@ REGISTER_NAME(257, Zlib)
 REGISTER_NAME(258, Gzip)
 REGISTER_NAME(259, LZ4)
 REGISTER_NAME(260, Mobile)
+REGISTER_NAME(261, Oodle)
 
 // Online
 REGISTER_NAME(280,DGram)
diff --git a/Engine/Source/Runtime/Launch/Private/LaunchEngineLoop.cpp b/Engine/Source/Runtime/Launch/Private/LaunchEngineLoop.cpp
--- a/Engine/Source/Runtime/Launch/Private/LaunchEngineLoop.cpp
+++ b/Engine/Source/Runtime/Launch/Private/LaunchEngineLoop.cpp
@@ -63,6 +63,8 @@
 #include "Misc/NetworkVersion.h"
 #include "Templates/UniquePtr.h"
 
+#include "Compression/OodleDataCompression.h"
+
 #if !(IS_PROGRAM || WITH_EDITOR)
 #include "IPlatformFilePak.h"
 #endif
@@ -2379,6 +2381,9 @@ int32 FEngineLoop::PreInitPreStartupScreen(const TCHAR* CmdLine)
 	GLogConsole = GScopedLogConsole.Get();
 #endif
 
+	// init Oodle here
+	OodleDataCompressionStartupPreInit();
+
 	{
 		SCOPED_BOOT_TIMING("LoadPreInitModules");
 		LoadPreInitModules();
@@ -2723,7 +2728,7 @@ PRAGMA_ENABLE_DEPRECATION_WARNINGS
 	GIsDemoMode = FParse::Param(FCommandLine::Get(), TEXT("DEMOMODE"));
 #endif
 
-
+	// InitEngineTextLocalization loads Paks, needs Oodle to be setup before here
 	InitEngineTextLocalization();
 
 	bool bForceEnableHighDPI = false;
diff --git a/Engine/Source/Runtime/OodleDataCompression/OodleDataCompression.Build.cs b/Engine/Source/Runtime/OodleDataCompression/OodleDataCompression.Build.cs
new file mode 100644
--- /dev/null
+++ b/Engine/Source/Runtime/OodleDataCompression/OodleDataCompression.Build.cs
@@ -0,0 +1,107 @@
+// Copyright Epic Games, Inc. All Rights Reserved.
+
+using System.IO;
+using System.Collections.Generic;
+using System.Security.Cryptography;
+using UnrealBuildTool;
+
+//namespace UnrealBuildTool.Rules
+//{
+	public class OodleDataCompression : ModuleRules
+	{
+		protected virtual string OodleVersion { get { return "2.9.2"; } }
+
+		// Platform Extensions need to override these
+		protected virtual string LibRootDirectory { get { return ModuleDirectory; } }
+		protected virtual string ReleaseLibraryName { get { return null; } }
+		protected virtual string DebugLibraryName { get { return null; } }
+
+		protected virtual string SdkBaseDirectory { get { return Path.Combine(LibRootDirectory, "Sdks", OodleVersion); } }
+		protected virtual string LibDirectory { get { return Path.Combine(SdkBaseDirectory, "lib"); } }
+
+		protected virtual string IncludeDirectory { get { return Path.Combine(ModuleDirectory, "Sdks", OodleVersion, "include"); } }
+
+
+		public OodleDataCompression(ReadOnlyTargetRules Target) : base(Target)
+		{
+			Type = ModuleType.External;
+
+			ShortName = "OodleDataCompression";
+
+			PublicIncludePaths.Add(IncludeDirectory);
+
+			string ReleaseLib = null;
+			string DebugLib = null;
+			string PlatformDir = Target.Platform.ToString();
+
+			// turn on bAllowDebugLibrary if you need to debug a problem with Oodle
+			bool bAllowDebugLibrary = false;
+			bool bUseDebugLibrary = bAllowDebugLibrary && Target.Configuration == UnrealTargetConfiguration.Debug && Target.bDebugBuildsActuallyUseDebugCRT;
+
+			bool bSkipLibrarySetup = false;
+
+			if (Target.IsInPlatformGroup(UnrealPlatformGroup.Windows))
+			{
+				ReleaseLib = "oo2core_win64.lib";
+				DebugLib = "oo2core_win64_debug.lib";
+				PlatformDir = "Win64";
+			}
+			else if (Target.Platform == UnrealTargetPlatform.Mac)
+			{
+				ReleaseLib = "liboo2coremac64.a";
+				DebugLib = "liboo2coremac64_dbg.a";
+			}
+			else if (Target.Platform == UnrealTargetPlatform.Linux)
+			{
+				ReleaseLib = "liboo2corelinux64.a";
+				DebugLib = "liboo2corelinux64_dbg.a";
+			}
+			else if (Target.Platform == UnrealTargetPlatform.LinuxArm64)
+			{
+				ReleaseLib = "liboo2corelinuxarm64.a";
+				DebugLib = "liboo2corelinuxarm64_dbg.a";
+			}
+			else if (Target.Platform == UnrealTargetPlatform.Android)
+			{
+				// we have to add multiple libraries, so don't do the normal processing below
+				bSkipLibrarySetup = true;
+				string dbg = bUseDebugLibrary ? "_dbg" : "";
+
+				PublicAdditionalLibraries.Add(Path.Combine(LibDirectory, PlatformDir, "arm64-v8a/liboo2coreandroid" + dbg + ".a"));
+				PublicAdditionalLibraries.Add(Path.Combine(LibDirectory, PlatformDir, "armeabi-v7a/liboo2coreandroid" + dbg + ".a"));
+				PublicAdditionalLibraries.Add(Path.Combine(LibDirectory, PlatformDir, "x86/liboo2coreandroid" + dbg + ".a"));
+				PublicAdditionalLibraries.Add(Path.Combine(LibDirectory, PlatformDir, "x86_64/liboo2coreandroid" + dbg + ".a"));
+			}
+			else if (Target.Platform == UnrealTargetPlatform.IOS)
+			{
+				ReleaseLib = "liboo2coreios.a";
+				DebugLib = "liboo2coreios_dbg.a";
+			}
+			else if (Target.Platform == UnrealTargetPlatform.TVOS)
+			{
+				ReleaseLib = "liboo2coretvos.a";
+				DebugLib = "liboo2coretvos_dbg.a";
+			}
+			else
+			{
+				// the subclass will return the library names
+				ReleaseLib = ReleaseLibraryName;
+				DebugLib = DebugLibraryName;
+				// platform extensions don't need the Platform directory under lib
+				PlatformDir = "";
+			}
+
+			if (!bSkipLibrarySetup)
+			{
+				// combine everything and make sure it was set up properly
+				string LibraryToLink = bUseDebugLibrary ? DebugLib : ReleaseLib;
+				if (LibraryToLink == null)
+				{
+					throw new BuildException("Platform {0} doesn't have OodleData libraries properly set up.", Target.Platform);
+				}
+
+				PublicAdditionalLibraries.Add(Path.Combine(LibDirectory, PlatformDir, LibraryToLink));
+			}
+		}
+	}
+//}
diff --git a/Engine/Source/Runtime/OodleDataCompression/Sdks/2.9.2/include/oodle2.h b/Engine/Source/Runtime/OodleDataCompression/Sdks/2.9.2/include/oodle2.h
new file mode 100644
--- /dev/null
+++ b/Engine/Source/Runtime/OodleDataCompression/Sdks/2.9.2/include/oodle2.h
@@ -0,0 +1,1644 @@
+
+//===================================================
+// Oodle2 Core header
+// (C) Copyright 1994-2021 Epic Games Tools LLC
+//===================================================
+
+#ifndef __OODLE2_H_INCLUDED__
+#define __OODLE2_H_INCLUDED__
+
+#ifndef OODLE2_PUBLIC_HEADER
+#define OODLE2_PUBLIC_HEADER 1
+#endif
+
+#ifndef __OODLE2BASE_H_INCLUDED__
+#include "oodle2base.h"
+#endif
+
+#ifdef _MSC_VER
+#pragma pack(push, Oodle, 8)
+
+#pragma warning(push)
+#pragma warning(disable : 4127) // conditional is constant
+#endif
+
+// header version :
+//  the DLL is incompatible when MAJOR is bumped
+//  MINOR is for internal revs and bug fixes that don't affect API compatibility
+#define OODLE2_VERSION_MAJOR            9
+#define OODLE2_VERSION_MINOR            2
+
+// OodleVersion string is 1 . MAJOR . MINOR
+//  don't make it from macros cuz the doc tool has to parse the string literal
+
+#define OodleVersion "2.9.2"    /*
+*/
+
+//-----------------------------------------------------
+// OodleLZ
+
+#if 0
+#define OODLE_ALLOW_DEPRECATED_COMPRESSORS /* If you need to encode with the deprecated compressors, define this before including oodle2.h
+
+    You may still decode with them without defining this.
+*/
+#endif
+
+// Default verbosity selection of 0 will not even log when it sees corruption
+typedef enum OodleLZ_Verbosity
+{
+    OodleLZ_Verbosity_None = 0,
+    OodleLZ_Verbosity_Minimal = 1,
+    OodleLZ_Verbosity_Some = 2,
+    OodleLZ_Verbosity_Lots = 3,
+    OodleLZ_Verbosity_Force32 = 0x40000000
+} OodleLZ_Verbosity;
+/* Verbosity of LZ functions
+    LZ functions print information to the function set by $OodleCore_Plugins_SetPrintf
+    or $OodleXLog_Printf if using OodleX.
+*/
+
+OO_COMPILER_ASSERT( sizeof(OodleLZ_Verbosity) == 4 );
+
+typedef enum OodleLZ_Compressor
+{
+    OodleLZ_Compressor_Invalid = -1,
+    OodleLZ_Compressor_None = 3,  // None = memcpy, pass through uncompressed bytes
+
+    // NEW COMPRESSORS :
+    OodleLZ_Compressor_Kraken = 8,    // Fast decompression and high compression ratios, amazing!
+    OodleLZ_Compressor_Leviathan = 13,// Leviathan = Kraken's big brother with higher compression, slightly slower decompression.
+    OodleLZ_Compressor_Mermaid = 9,   // Mermaid is between Kraken & Selkie - crazy fast, still decent compression.
+    OodleLZ_Compressor_Selkie = 11,   // Selkie is a super-fast relative of Mermaid.  For maximum decode speed.
+    OodleLZ_Compressor_Hydra = 12,    // Hydra, the many-headed beast = Leviathan, Kraken, Mermaid, or Selkie (see $OodleLZ_About_Hydra)
+
+#ifdef OODLE_ALLOW_DEPRECATED_COMPRESSORS
+    OodleLZ_Compressor_BitKnit = 10, // no longer supported as of Oodle 2.9.0
+    OodleLZ_Compressor_LZB16 = 4, // DEPRECATED but still supported
+    OodleLZ_Compressor_LZNA = 7,  // no longer supported as of Oodle 2.9.0
+    OodleLZ_Compressor_LZH = 0,   // no longer supported as of Oodle 2.9.0
+    OodleLZ_Compressor_LZHLW = 1, // no longer supported as of Oodle 2.9.0
+    OodleLZ_Compressor_LZNIB = 2, // no longer supported as of Oodle 2.9.0
+    OodleLZ_Compressor_LZBLW = 5, // no longer supported as of Oodle 2.9.0
+    OodleLZ_Compressor_LZA = 6,   // no longer supported as of Oodle 2.9.0
+#endif
+
+    OodleLZ_Compressor_Count = 14,
+    OodleLZ_Compressor_Force32 = 0x40000000
+} OodleLZ_Compressor;
+/* Selection of compression algorithm.
+
+    Each compressor provides a different balance of speed vs compression ratio.
+
+    New Oodle users should only use the new sea monster family of compressors.
+
+    The OODLE_ALLOW_DEPRECATED_COMPRESSORS set of compressors is no longer supported
+    as of Oodle 2.9.0 ; see $Oodle_FAQ_deprecated_compressors
+
+    The sea monsters are all fuzz safe and use whole-block quantum (not the 16k quantum)
+    ($OodleLZ_Compressor_UsesWholeBlockQuantum)
+
+    If you need to encode the deprecated compressors, define $OODLE_ALLOW_DEPRECATED_COMPRESSORS before
+    including oodle2.h  
+
+    See $Oodle_FAQ_WhichLZ for a quick FAQ on which compressor to use
+
+    See $OodleLZ_About for discussion of how to choose a compressor.
+*/
+
+OO_COMPILER_ASSERT( sizeof(OodleLZ_Compressor) == 4 );
+
+typedef enum OodleLZ_PackedRawOverlap
+{
+    OodleLZ_PackedRawOverlap_No = 0,
+    OodleLZ_PackedRawOverlap_Yes = 1,
+    OodleLZ_PackedRawOverlap_Force32 = 0x40000000
+} OodleLZ_PackedRawOverlap;
+/* Bool enum
+*/
+
+typedef enum OodleLZ_CheckCRC
+{
+    OodleLZ_CheckCRC_No = 0,
+    OodleLZ_CheckCRC_Yes = 1,
+    OodleLZ_CheckCRC_Force32 = 0x40000000
+} OodleLZ_CheckCRC;
+/* Bool enum for the LZ decoder - should it check CRC before decoding or not?
+
+    NOTE : the CRC's in the LZH decompress checks are the CRC's of the *compressed* bytes.  This allows checking the CRc
+    prior to decompression, so corrupted data cannot be fed to the compressor.
+
+    To use OodleLZ_CheckCRC_Yes, the compressed data must have been made with $(OodleLZ_CompressOptions:sendQuantumCRCs) set to true.
+
+    If you want a CRC of the raw bytes, there is one optionally stored in the $OodleLZ_SeekTable and can be confirmed with
+    $OodleLZ_CheckSeekTableCRCs
+*/
+
+
+typedef enum OodleLZ_Profile
+{
+    OodleLZ_Profile_Main=0,         // Main profile (all current features allowed)
+    OodleLZ_Profile_Reduced=1,      // Reduced profile (Kraken only, limited feature set)
+    OodleLZ_Profile_Force32 = 0x40000000
+} OodleLZ_Profile;
+/* Decode profile to target */
+
+// Not flagged for idoc and done using a #define since it's internal (testing) use only
+#define OodleLZ_Profile_Internal_Custom ((OodleLZ_Profile)100)
+
+OO_COMPILER_ASSERT( sizeof(OodleLZ_Profile) == 4 );
+
+typedef enum OodleDecompressCallbackRet
+{
+    OodleDecompressCallbackRet_Continue=0,
+    OodleDecompressCallbackRet_Cancel=1,
+    OodleDecompressCallbackRet_Invalid=2,
+    OodleDecompressCallbackRet_Force32 = 0x40000000
+} OodleDecompressCallbackRet;
+/* Return value for $OodleDecompressCallback
+    return OodleDecompressCallbackRet_Cancel to abort the in-progress decompression
+*/
+
+OODEFFUNC typedef OodleDecompressCallbackRet (OODLE_CALLBACK OodleDecompressCallback)(void * userdata, const OO_U8 * rawBuf,OO_SINTa rawLen,const OO_U8 * compBuf,OO_SINTa compBufferSize , OO_SINTa rawDone, OO_SINTa compUsed);
+/* User-provided callback for decompression
+
+    $:userdata  the data you passed for _pcbData_
+    $:rawBuf    the decompressed buffer
+    $:rawLen    the total decompressed length
+    $:compBuf   the compressed buffer
+    $:compBufferSize  the total compressed length
+    $:rawDone   number of bytes in rawBuf decompressed so far
+    $:compUsed  number of bytes in compBuf consumed so far
+
+    OodleDecompressCallback is called incrementally during decompression.
+*/
+
+typedef enum OodleLZ_CompressionLevel
+{
+    OodleLZ_CompressionLevel_None=0,        // don't compress, just copy raw bytes
+    OodleLZ_CompressionLevel_SuperFast=1,   // super fast mode, lower compression ratio
+    OodleLZ_CompressionLevel_VeryFast=2,    // fastest LZ mode with still decent compression ratio
+    OodleLZ_CompressionLevel_Fast=3,        // fast - good for daily use
+    OodleLZ_CompressionLevel_Normal=4,      // standard medium speed LZ mode
+
+    OodleLZ_CompressionLevel_Optimal1=5,    // optimal parse level 1 (faster optimal encoder)
+    OodleLZ_CompressionLevel_Optimal2=6,    // optimal parse level 2 (recommended baseline optimal encoder)
+    OodleLZ_CompressionLevel_Optimal3=7,    // optimal parse level 3 (slower optimal encoder)
+    OodleLZ_CompressionLevel_Optimal4=8,    // optimal parse level 4 (very slow optimal encoder)
+    OodleLZ_CompressionLevel_Optimal5=9,    // optimal parse level 5 (don't care about encode speed, maximum compression)
+
+    OodleLZ_CompressionLevel_HyperFast1=-1, // faster than SuperFast, less compression
+    OodleLZ_CompressionLevel_HyperFast2=-2, // faster than HyperFast1, less compression
+    OodleLZ_CompressionLevel_HyperFast3=-3, // faster than HyperFast2, less compression
+    OodleLZ_CompressionLevel_HyperFast4=-4, // fastest, less compression
+
+    // aliases :
+    OodleLZ_CompressionLevel_HyperFast=OodleLZ_CompressionLevel_HyperFast1, // alias hyperfast base level
+    OodleLZ_CompressionLevel_Optimal = OodleLZ_CompressionLevel_Optimal2,   // alias optimal standard level
+    OodleLZ_CompressionLevel_Max     = OodleLZ_CompressionLevel_Optimal5,   // maximum compression level
+    OodleLZ_CompressionLevel_Min     = OodleLZ_CompressionLevel_HyperFast4, // fastest compression level
+
+    OodleLZ_CompressionLevel_Force32 = 0x40000000,
+    OodleLZ_CompressionLevel_Invalid = OodleLZ_CompressionLevel_Force32
+} OodleLZ_CompressionLevel;
+/* Selection of compression encoder complexity
+
+    Higher numerical value of CompressionLevel = slower compression, but smaller compressed data.
+
+    The compressed stream is always decodable with the same decompressors.
+    CompressionLevel controls the amount of work the encoder does to find the best compressed bit stream.
+    CompressionLevel does not primary affect decode speed, it trades off encode speed for compressed bit stream quality.
+
+    I recommend starting with OodleLZ_CompressionLevel_Normal, then try up or down if you want
+    faster encoding or smaller output files.
+
+    The Optimal levels are good for distribution when you compress rarely and decompress often;
+    they provide very high compression ratios but are slow to encode.  Optimal2 is the recommended level
+    to start with of the optimal levels.
+    Optimal4 and 5 are not recommended for common use, they are very slow and provide the maximum compression ratio,
+    but the gain over Optimal3 is usually small.
+
+    The HyperFast levels have negative numeric CompressionLevel values.
+    They are faster than SuperFast for when you're encoder CPU time constrained or want
+    something closer to symmetric compression vs. decompression time.
+    The HyperFast levels are currently only available in Kraken, Mermaid & Selkie.
+    Higher levels of HyperFast are faster to encode, eg. HyperFast4 is the fastest.
+
+    The CompressionLevel does not affect decode speed much.  Higher compression level does not mean
+    slower to decode.  To trade off decode speed vs ratio, use _spaceSpeedTradeoffBytes_ in $OodleLZ_CompressOptions
+
+*/
+
+OO_COMPILER_ASSERT( sizeof(OodleLZ_CompressionLevel) == 4 );
+
+typedef enum OodleLZ_Jobify
+{
+    OodleLZ_Jobify_Default=0,       // Use compressor default for level of internal job usage
+    OodleLZ_Jobify_Disable=1,       // Don't use jobs at all
+    OodleLZ_Jobify_Normal=2,        // Try to balance parallelism with increased memory usage
+    OodleLZ_Jobify_Aggressive=3,    // Maximize parallelism even when doing so requires large amounts of memory
+    OodleLZ_Jobify_Count=4,
+
+    OodleLZ_Jobify_Force32 = 0x40000000,
+} OodleLZ_Jobify;
+/* Controls the amount of internal threading in $OodleLZ_Compress calls
+
+    Once you install a pluggable job system via $OodleCore_Plugins_SetJobSystem, Oodle can internally break
+    heavy-weight compression tasks into smaller jobs that can run in parallel.  This can speed up
+    compression of large blocks of data at Optimal1 and higher levels substantially.
+
+    The trade-off is that running more jobs concurrently rather than sequentially can greatly increase
+    memory requirements when there are multiple outstanding memory-intensive jobs.
+
+    OodleLZ_Jobify_Default lets the compressor decide; typically compressors will default to "Normal"
+    when a pluggable job system has been installed, and "Disable" otherwise.
+
+    OodleLZ_Jobify_Disable disables use of internal jobs entirely; all compression work is done on
+    the calling thread.  This minimizes the amount of memory used, and is also appropriate when you're
+    getting parallelism in other ways, e.g. by running OodleLZ_Compress on many threads yourself.
+
+    OodleLZ_Jobify_Normal uses jobs to increase compressor parallelism and speeds up compression of
+    large blocks of data, but avoids handing out many concurrent jobs for tasks that are memory-intensive.
+
+    OodleLZ_Jobify_Aggressive will use concurrent jobs even for highly memory-intensive tasks.  This
+    can speed up things further, but at a potentially significant increase in the amount of memory used
+    by Oodle.
+
+*/
+
+#define OODLELZ_LOCALDICTIONARYSIZE_MAX     (1<<30) /* Maximum value of maxLocalDictionarySize in OodleLZ_CompressOptions
+*/
+
+#define OODLELZ_SPACESPEEDTRADEOFFBYTES_DEFAULT (256) /* Default value of spaceSpeedTradeoffBytes in OodleLZ_CompressOptions
+    Changes how the encoder makes decisions in the bit stream
+    Higher spaceSpeedTradeoffBytes favors decode speed more (larger compressed files)
+    Lower spaceSpeedTradeoffBytes favors smaller compressed files (slower decoder)
+    Goes in a power of 2 scale; so try 64,128 and 512,1024
+    (OODLELZ_SPACESPEEDTRADEOFFBYTES_DEFAULT/2) or (OODLELZ_SPACESPEEDTRADEOFFBYTES_DEFAULT*2)
+*/
+
+
+typedef OOSTRUCT OodleLZ_CompressOptions
+{
+    OO_U32              unused_was_verbosity;               // unused ; was verbosity (set to zero)
+    OO_S32              minMatchLen;        // minimum match length ; cannot be used to reduce a compressor's default MML, but can be higher.  On some types of data, a large MML (6 or 8) is a space-speed win.
+    OO_BOOL             seekChunkReset;     // whether chunks should be independent, for seeking and parallelism
+    OO_S32              seekChunkLen;       // length of independent seek chunks (if seekChunkReset) ; must be a power of 2 and >= $OODLELZ_BLOCK_LEN ; you can use $OodleLZ_MakeSeekChunkLen
+    OodleLZ_Profile     profile;            // decoder profile to target (set to zero)
+    OO_S32              dictionarySize;     // sets a maximum offset for matches, if lower than the maximum the format supports.  <= 0 means infinite (use whole buffer).  Often power of 2 but doesn't have to be.
+    OO_S32              spaceSpeedTradeoffBytes;  // this is a number of bytes; I must gain at least this many bytes of compressed size to accept a speed-decreasing decision
+    OO_S32              unused_was_maxHuffmansPerChunk;  //  unused ; was maxHuffmansPerChunk
+    OO_BOOL             sendQuantumCRCs;    // should the encoder send a CRC of each compressed quantum, for integrity checks; this is necessary if you want to use OodleLZ_CheckCRC_Yes on decode
+    OO_S32              maxLocalDictionarySize;  // (Optimals) size of local dictionary before needing a long range matcher.  This does not set a window size for the decoder; it's useful to limit memory use and time taken in the encoder.  maxLocalDictionarySize must be a power of 2.  Must be <= OODLELZ_LOCALDICTIONARYSIZE_MAX
+    OO_BOOL             makeLongRangeMatcher;   // (Optimals) should the encoder find matches beyond maxLocalDictionarySize using an LRM
+    OO_S32              matchTableSizeLog2; //(non-Optimals)  when variable, sets the size  of the match finder structure (often a hash table) ; use 0 for the compressor's default
+
+    OodleLZ_Jobify      jobify;         // controls internal job usage by compressors
+    void *              jobifyUserPtr;  // user pointer passed through to RunJob and WaitJob callbacks
+
+    OO_S32              farMatchMinLen; // far matches must be at least this len
+    OO_S32              farMatchOffsetLog2; // if not zero, the log2 of an offset that must meet farMatchMinLen
+
+    OO_U32              reserved[4];   // reserved space for adding more options; zero these!
+} OodleLZ_CompressOptions;
+/* Options for the compressor
+
+    Typically filled by calling $OodleLZ_CompressOptions_GetDefault , then individual options may be modified, like :
+
+    OodleLZ_CompressOptions my_options = *OodleLZ_CompressOptions_GetDefault()
+
+    To ensure you have set up the options correctly, call $OodleLZ_CompressOptions_Validate.
+
+    _unused_was_verbosity_ : place holder, set to zero
+
+    _minMatchLen_ : rarely useful.  Default value of 0 means let the compressor decide.  On some types of data,
+    bumping this up to 4,6, or 8 can improve decode speed with little effect on compression ratio.  Most of the
+    Oodle compressors use a default MML of 4 at levels below 7, and MML 3 at levels >= 7.  If you want to keep MML 4
+    at the higher levels, set _minMatchLen_ here to 4.  _minMatchLen_ cannot be used to reduce the base MML of the compressor, only to increase it.
+
+    _seekChunkReset_ must be true if you want the decode to be able to run "Wide", with pieces that can be
+    decoded independently (not keeping previous pieces in memory for match references).
+
+    _seekChunkLen_ : length of independent seek chunks (if seekChunkReset) ; must be a power of 2 and >= $OODLELZ_BLOCK_LEN ; you can use $OodleLZ_MakeSeekChunkLen
+
+    _profile_ : tells the encoder to target alternate bitstream profile.  Default value of zero for normal use.
+
+    _dictionarySize_ : limits the encoder to partial buffer access for matches.  Can be useful for decoding incrementally
+    without keeping the entire output buffer in memory.
+
+    _spaceSpeedTradeoffBytes_ is a way to trade off compression ratio for decode speed.  If you make it smaller,
+    you get more compression ratio and slower decodes.  It's the number of bytes that a decision must save to
+    be worth a slower decode.  Default is 256 (OODLELZ_SPACESPEEDTRADEOFFBYTES_DEFAULT).  So that means the encoder must be able to save >= 256 bytes to
+    accept something that will slow down decoding (like adding another Huffman table).  The typical range is
+    64-1024.
+
+    Lower _spaceSpeedTradeoffBytes_ = more compression, slower decode
+    Higher _spaceSpeedTradeoffBytes_ = less compression, faster decode
+
+    _spaceSpeedTradeoffBytes_ is the primary parameter for controlling Hydra.  The default value of 256 will make
+    Hydra decodes that are just a little bit faster than Kraken.  You get Kraken speeds around 200, and Mermaid
+    speeds around 1200.
+
+    At the extreme, a _spaceSpeedTradeoffBytes_ of zero would mean all you care about is compression ratio, not decode
+    speed, you want the encoder to make the smallest possible output.  (you cannot actually set zero, as zero values
+    always mean "use default" in this struct; you never really want zero anyway)
+    Generally _spaceSpeedTradeoffBytes_ below 16 provides diminishing gains in size with pointless decode speed loss.
+
+    _spaceSpeedTradeoffBytes_ is on sort of powers of 2 scale, so you might want to experiment with 32,64,128,256,512
+
+    _spaceSpeedTradeoffBytes_ outside the range [16 - 2048] is not recommended.
+
+    _unused_was_maxHuffmansPerChunk_ : place holder, set to zero
+
+    _sendQuantumCRCs_ : send hashes of the compressed data to verify in the decoder; not recommended, if you need data
+    verification, use your own system outside of Oodle.  DEPRECATED, not recommended.  For backwards compatibility only.
+
+    _maxLocalDictionarySize_ : only applies to optimal parsers at level >= Optimal2.  This limits the encoder memory use.
+    Making it larger = more compression, higher memory use.  Matches within maxLocalDictionarySize are found exactly,
+    outside the maxLocalDictionarySize window an approximate long range matcher is used.
+
+    _makeLongRangeMatcher_ : whether an LRM should be used to find matches outside the _maxLocalDictionarySize_ window
+    (Optimal levels only)
+
+    _matchTableSizeLog2_ : for non-optimal levels (level <= Normal), controls the hash table size.  Making this very
+    small can sometimes boost encoder speed.  For the very fastest encoding, use the SuperFast level and change
+    _matchTableSizeLog2_ to 12 or 13.
+
+    _matchTableSizeLog2_ should usually be left zero to use the encoder's default
+
+    _matchTableSizeLog2_ allows you to limit memory use of the non-Optimal encoder levels.  Memory use is roughly
+    ( 1 MB + 4 << matchTableSizeLog2 )
+
+    _jobify_ tells compressors how to use internal jobs for compression tasks.  Jobs can be run in parallel using the
+    job system plugins set with $OodleCore_Plugins_SetJobSystem.  Not all compressors or compression level support
+    jobs, but the slower ones generally do.  The default value of jobify is to use a thread system if one is installed.
+
+    _farMatchMinLen_ and _farMatchOffsetLog2_ can be used to tune the encoded stream for a known cache size on the
+    decoding hardware.  If set, then offsets with log2 greater or each to _farMatchOffsetLog2_ must have a minimum
+    length of _farMatchMinLen_.  For example to target a machine with a 2 MB cache, set _farMatchOffsetLog2_ to 21,
+    and _farMatchMinLen_ to something large, like 16 or 20.
+
+    Without _farMatchMinLen_ and _farMatchOffsetLog2_ set, the Oodle encoders tune for a blend of cache sizes that works
+    well on most machines.  _dictionarySize_ can also be used to tune for cache size, but cuts off all matches
+    beyond a certain distance.  That may be more appropriate when you don't want to go out of cache at all.
+    _farMatchMinLen_ can only be used to make the standard blend target more restrictive; it can reduce the target cache size
+    but can't make it larger (or it can raise min match len outside cache but can't make it shorter).
+
+    For help on setting up OodleLZ_CompressOptions contact support at [REDACTED_EMAIL]
+
+    NOTE : fields you do not set should always be zero initialized.  In particular the _reserved_ fields should be zeroed.
+    Zero always means "use default" and is a future-portable initialization value.
+
+    If you set fields to zero to mean "use default" you can call $OodleLZ_CompressOptions_Validate to change them
+    to default values.  This is done automatically internally if you don't do it explicitly.
+
+*/
+
+typedef enum OodleLZ_Decode_ThreadPhase
+{
+    OodleLZ_Decode_ThreadPhase1 = 1,
+    OodleLZ_Decode_ThreadPhase2 = 2,
+    OodleLZ_Decode_ThreadPhaseAll = 3,
+    OodleLZ_Decode_Unthreaded = OodleLZ_Decode_ThreadPhaseAll
+} OodleLZ_Decode_ThreadPhase;
+/* ThreadPhase for threaded Oodle decode
+
+    Check $OodleLZ_Compressor_CanDecodeThreadPhased
+    (currently only used by Kraken)
+
+    See $OodleLZ_About_ThreadPhasedDecode
+
+*/
+
+typedef enum OodleLZ_FuzzSafe
+{
+    OodleLZ_FuzzSafe_No = 0,
+    OodleLZ_FuzzSafe_Yes = 1
+} OodleLZ_FuzzSafe;
+/* OodleLZ_FuzzSafe (deprecated)
+
+    About fuzz safety:
+
+    Fuzz Safe decodes will not crash on corrupt data.  They may or may not return failure, and produce garbage output.
+
+    Fuzz safe decodes will not read out of bounds.  They won't put data on the stack or previously in memory
+    into the output buffer.
+
+    As of Oodle 2.9.0 all compressors supported are fuzzsafe, so OodleLZ_FuzzSafe_Yes should always be used and this
+    enum is deprecated.
+
+*/
+
+#define OODLELZ_BLOCK_LEN   (1<<18) /* The number of raw bytes per "seek chunk"
+    Seek chunks can be decompressed independently if $(OodleLZ_CompressOptions:seekChunkReset) is set.
+*/
+
+#define OODLELZ_BLOCK_MAXIMUM_EXPANSION (2)
+#define OODLELZ_BLOCK_MAX_COMPLEN       (OODLELZ_BLOCK_LEN+OODLELZ_BLOCK_MAXIMUM_EXPANSION) /* Maximum expansion per $OODLELZ_BLOCK_LEN is 1 byte.
+    Note that the compressed buffer must be allocated bigger than this (use $OodleLZ_GetCompressedBufferSizeNeeded)
+*/
+
+#define OODLELZ_QUANTUM_LEN     (1<<14) /* Minimum decompression quantum (for old legacy codecs only)
+
+    Deprecated.
+
+    The new sea monster family of compressors use a whole block quantum (OODLELZ_BLOCK_LEN).
+    Check $OodleLZ_Compressor_UsesWholeBlockQuantum
+*/
+
+// 5 byte expansion per-quantum with CRC's
+#define OODLELZ_QUANTUM_MAXIMUM_EXPANSION   (5)
+
+#define OODLELZ_QUANTUM_MAX_COMPLEN     (OODLELZ_QUANTUM_LEN+OODLELZ_QUANTUM_MAXIMUM_EXPANSION)
+
+#define OODLELZ_SEEKCHUNKLEN_MIN        OODLELZ_BLOCK_LEN
+#define OODLELZ_SEEKCHUNKLEN_MAX        (1<<29) // half GB
+
+typedef OOSTRUCT OodleLZ_DecodeSome_Out
+{
+    OO_S32 decodedCount;    // number of uncompressed bytes decoded
+    OO_S32 compBufUsed; // number of compressed bytes consumed
+
+
+    OO_S32 curQuantumRawLen;  // tells you the current quantum size. you must have at least this much room available in the output buffer to be able to decode anything.
+    OO_S32 curQuantumCompLen; // if you didn't pass in enough data, nothing will decode (decodedCount will be 0), and this will tell you how much is needed
+} OodleLZ_DecodeSome_Out;
+/* Output value of $OodleLZDecoder_DecodeSome
+*/
+
+//---------------------------------------------
+
+//=======================================================
+
+typedef OOSTRUCT OodleLZ_SeekTable
+{
+    OodleLZ_Compressor  compressor;             // which compressor was used
+    OO_BOOL             seekChunksIndependent;  // are the seek chunks independent, or must they be decompressed in sequence
+
+    OO_S64              totalRawLen;    // total uncompressed data lenth
+    OO_S64              totalCompLen;   // sum of seekChunkCompLens
+
+    OO_S32              numSeekChunks;  // derived from rawLen & seekChunkLen
+    OO_S32              seekChunkLen;   // multiple of OODLELZ_BLOCK_LEN
+
+    OO_U32 *            seekChunkCompLens;  // array of compressed lengths of seek chunks
+    OO_U32 *            rawCRCs;            // crc of the raw bytes of the chunk (optional; NULL unless $OodleLZSeekTable_Flags_MakeRawCRCs was specified)
+} OodleLZ_SeekTable;
+
+typedef enum OodleLZSeekTable_Flags
+{
+    OodleLZSeekTable_Flags_None  = 0,       // default
+    OodleLZSeekTable_Flags_MakeRawCRCs = 1,  // make the _rawCRCs_ member of $OodleLZ_SeekTable
+    OodleLZSeekTable_Flags_Force32 = 0x40000000
+} OodleLZSeekTable_Flags;
+
+//=====================================================
+
+
+typedef OOSTRUCT OodleConfigValues
+{
+    OO_S32 m_OodleLZ_LW_LRM_step;           // LZHLW LRM : bytes between LRM entries
+    OO_S32 m_OodleLZ_LW_LRM_hashLength; // LZHLW LRM : bytes hashed for each LRM entries
+    OO_S32 m_OodleLZ_LW_LRM_jumpbits;       // LZHLW LRM : bits of hash used for jump table
+
+    OO_S32 m_OodleLZ_Decoder_Max_Stack_Size;    // if OodleLZ_Decompress needs to allocator a Decoder object, and it's smaller than this size, it's put on the stack instead of the heap
+    OO_S32 m_OodleLZ_Small_Buffer_LZ_Fallback_Size_Unused; // deprecated
+    OO_S32 m_OodleLZ_BackwardsCompatible_MajorVersion; // if you need to encode streams that can be read with an older version of Oodle, set this to the Oodle2 MAJOR version number that you need compatibility with.  eg to be compatible with oodle 2.7.3 you would put 7 here
+
+    OO_U32 m_oodle_header_version;  // = OODLE_HEADER_VERSION
+
+} OodleConfigValues;
+/* OodleConfigValues
+
+    Struct of user-settable low level config values.  See $Oodle_SetConfigValues.
+
+    May have different defaults per platform.
+*/
+
+OOFUNC1 void OOFUNC2 Oodle_GetConfigValues(OodleConfigValues * ptr);
+/* Get $OodleConfigValues
+
+    $:ptr   filled with OodleConfigValues
+
+    Gets the current $OodleConfigValues.
+
+    May be different per platform.
+*/
+
+OOFUNC1 void OOFUNC2 Oodle_SetConfigValues(const OodleConfigValues * ptr);
+/* Set $OodleConfigValues
+
+    $:ptr   your desired OodleConfigValues
+
+    Sets the global $OodleConfigValues from your struct.
+
+    You should call $Oodle_GetConfigValues to fill the struct, then change the values you
+    want to change, then call $Oodle_SetConfigValues.
+
+    This should generally be done before doing anything with Oodle (eg. even before OodleX_Init).
+    Changing OodleConfigValues after Oodle has started has undefined effects.
+*/
+
+typedef enum Oodle_UsageWarnings
+{
+    Oodle_UsageWarnings_Enabled = 0,
+    Oodle_UsageWarnings_Disabled = 1,
+    Oodle_UsageWarnings_Force32 = 0x40000000
+} Oodle_UsageWarnings;
+/* Whether Oodle usage warnings are enable or disabled. */
+
+OOFUNC1 void OOFUNC2 Oodle_SetUsageWarnings(Oodle_UsageWarnings state);
+/* Enables or disables Oodle usage warnings.
+
+    $:state    whether usage warnings should be enabled or disabled.
+
+   Usage warnings are enabled by default and try to be low-noise, but in case you want to
+   disable them, this is how.
+
+   This should generally be done once at startup.  Setting this state while there are Oodle
+   calls running on other threads has undefined results.
+*/
+
+// function pointers to mallocs needed :
+
+OODEFFUNC typedef void * (OODLE_CALLBACK t_fp_OodleCore_Plugin_MallocAligned)( OO_SINTa bytes, OO_S32 alignment);
+/* Function pointer type for OodleMallocAligned
+
+    $:bytes     number of bytes to allocate
+    $:alignment required alignment of returned pointer
+    $:return    pointer to memory allocated (must not be NULL)
+
+    _alignment_ will always be a power of two
+
+    _alignment_ will always be >= $OODLE_MALLOC_MINIMUM_ALIGNMENT
+
+*/
+
+OODEFFUNC typedef void (OODLE_CALLBACK t_fp_OodleCore_Plugin_Free)( void * ptr );
+/* Function pointer type for OodleFree
+
+    $:return    pointer to memory to free
+
+*/
+
+OOFUNC1 void OOFUNC2 OodleCore_Plugins_SetAllocators(
+    t_fp_OodleCore_Plugin_MallocAligned * fp_OodleMallocAligned,
+    t_fp_OodleCore_Plugin_Free * fp_OodleFree);
+/* Set the function pointers for allocation needed by Oodle2 Core
+
+    If these are not set, the default implementation on most platforms uses the C stdlib.
+    On Microsoft platforms the default implementation uses HeapAlloc.
+
+    These must not be changed once they are set!  Set them once then don't change them.
+
+    NOTE: if you are using Oodle Ext, do NOT call this.  OodleX_Init will install an allocator for Oodle Core.  Do not mix your own allocator with the OodleX allocator.  See $OodleXAPI_Malloc.
+
+    If you want to ensure that Oodle is not doing any allocations, you can call OodleCore_Plugins_SetAllocators(NULL,NULL);
+    If you do that, then any time Oodle needs to allocate memory internally, it will stop the process.
+    It is STRONGLY not recommended that you ship that way.  You can verify that Oodle is not allocating, but then leave some
+    fallback allocator installed when you actually ship just in case.
+
+    Also note that on many consoles the standard allocation practices may not
+    leave much heap memory for the C stdlib malloc.  In this case Oodle may fail to allocate.
+
+*/
+
+OODEFFUNC typedef OO_U64 (OODLE_CALLBACK t_fp_OodleCore_Plugin_RunJob)( t_fp_Oodle_Job * fp_job, void * job_data , OO_U64 * dependencies, int num_dependencies, void * user_ptr );
+/* Function pointer type for OodleCore_Plugins_SetJobSystem
+
+    $:dependencies      array of handles of other pending jobs. All guaranteed to be nonzero.
+    $:num_dependencies  number of dependencies. Guaranteed to be no more than OODLE_JOB_MAX_DEPENDENCIES.
+    $:user_ptr          is passed through from the OodleLZ_CompressOptions.
+    $:return            handle to the async job, or 0 if it was run synchronously
+
+    RunJob will call fp_job(job_data)
+
+    it may be done on a thread, or it may run the function synchronously and return 0, indicating the job is already done.
+    The returned OO_U64 is a handle passed to WaitJob, unless it is 0, in which case WaitJob won't get called.
+
+    fp_job should not run until all the dependencies are done.  This function should not delete the dependencies.
+
+    RunJob must be callable from within an Oodle Job, i.e. jobs may spawn their own sub-jobs directly.
+    However, the matching WaitJob calls will only ever occur on the thread that called the
+    internally threaded Oodle API function.
+
+    See $Oodle_About_Job_Threading_Plugins
+*/
+
+OODEFFUNC typedef void (OODLE_CALLBACK t_fp_OodleCore_Plugin_WaitJob)( OO_U64 job_handle, void * user_ptr );
+/* Function pointer type for OodleCore_Plugins_SetJobSystem
+
+    $:job_handle    a job handle returned from RunJob. Never 0.
+    $:user_ptr      is passed through from the OodleLZ_CompressOptions.
+
+    Waits until the job specified by job_handle is done and cleans up any associated resources. Oodle
+    will call WaitJob exactly once for every RunJob call that didn't return 0.
+
+    If job_handle was already completed, this should clean it up without waiting.
+
+    A handle value should not be reused by another RunJob until WaitJob has been done with that value.
+
+    WaitJob will not be called from running jobs.  It will be only be called from the original thread that
+    invoked Oodle.  If you are running Oodle from a worker thread, ensure that that thread is allowed to wait
+    on other job threads.
+
+    See $Oodle_About_Job_Threading_Plugins
+*/
+
+OOFUNC1 void OOFUNC2 OodleCore_Plugins_SetJobSystem(
+    t_fp_OodleCore_Plugin_RunJob * fp_RunJob,
+    t_fp_OodleCore_Plugin_WaitJob * fp_WaitJob);
+/* DEPRECATED use OodleCore_Plugins_SetJobSystemAndCount instead
+
+    See $OodleCore_Plugins_SetJobSystemAndCount
+*/
+
+
+OOFUNC1 void OOFUNC2 OodleCore_Plugins_SetJobSystemAndCount(
+    t_fp_OodleCore_Plugin_RunJob * fp_RunJob,
+    t_fp_OodleCore_Plugin_WaitJob * fp_WaitJob,
+    int target_parallelism);
+/* Set the function pointers for async job system needed by Oodle2 Core
+
+    $:fp_RunJob     pointer to RunJob function
+    $:fp_WaitJob    pointer to WaitJob function
+    $:target_parallelism    goal of number of jobs to run simultaneously
+
+    If these are not set, the default implementation runs jobs synchronously on the calling thread.
+
+    These must not be changed once they are set!  Set them once then don't change them.
+
+    _target_parallelism_ allows you to tell Oodle how many Jobs it should try to keep in flight at once.
+    Depending on the operation it may not be able to split work into this many jobs (so fewer will be used),
+    but it will not exceed this count.
+
+    For Oodle Data LZ work, typically _target_parallelism_ is usually best at the number of hardware cores
+    not including hyper threads).
+
+    For Oodle Texture BCN encoding work, _target_parallelism_ is usually best as the full number of hyper cores.
+
+    In some cases you may wish to reduce _target_parallelism_ by 1 or 2 cores to leave some of the CPU free for
+    other work.
+
+    For example on a CPU with 16 cores and 32 hardware threads, for LZ work you might set _target_parallelism_ to 15
+    when calling OodleCorePlugins.  For BC7 encoding you might set _target_parallelism_ to 30 when calling OodleTexPlugins.
+
+    NOTE : if you are using Oodle Ext, do NOT call this.  OodleX_Init will install a job system for Oodle Core.
+    Note OodleX only installs automatically to Oodle Core, not Net or Tex.  See example_jobify.cpp for manual
+    plugin.
+
+    Replaces deprecated $OodleCore_Plugins_SetJobSystem
+
+    See $Oodle_About_Job_Threading_Plugins
+*/
+
+// the main func pointer for log :
+OODEFFUNC typedef void (OODLE_CALLBACK t_fp_OodleCore_Plugin_Printf)(int verboseLevel,const char * file,int line,const char * fmt,...);
+/* Function pointer to Oodle Core printf
+
+    $:verboseLevel  verbosity of the message; 0-2 ; lower = more important
+    $:file          C file that sent the message
+    $:line          C line that sent the message
+    $:fmt           vararg printf format string
+
+    The logging function installed here must parse varargs like printf.
+
+    _verboseLevel_ may be used to omit verbose messages.
+*/
+
+OOFUNC1 t_fp_OodleCore_Plugin_Printf * OOFUNC2 OodleCore_Plugins_SetPrintf(t_fp_OodleCore_Plugin_Printf * fp_rrRawPrintf);
+/* Install the callback used by Oodle Core for logging
+
+    $:fp_rrRawPrintf    function pointer to your log function; may be NULL to disable all logging
+    $:return            returns the previous function pointer
+
+    Use this function to install your own printf for Oodle Core.
+
+    The default implementation in debug builds, if you install nothing, uses the C stdio printf for logging.
+    On Microsoft platforms, it uses OutputDebugString and not stdio.
+
+    To disable all logging, call OodleCore_Plugins_SetPrintf(NULL)
+
+    WARNING : this function is NOT thread safe!  It should be done only once and done in a place where the caller can guarantee thread safety.
+
+    In the debug build of Oodle, you can install OodleCore_Plugin_Printf_Verbose to get more verbose logging
+
+*/
+
+OODEFFUNC typedef OO_BOOL (OODLE_CALLBACK t_fp_OodleCore_Plugin_DisplayAssertion)(const char * file,const int line,const char * function,const char * message);
+/* Function pointer to Oodle Core assert callback
+
+    $:file          C file that triggered the assert
+    $:line          C line that triggered the assert
+    $:function      C function that triggered the assert (may be NULL)
+    $:message       assert message
+    $:return        true to break execution at the assertion site, false to continue
+
+    This callback is called by Oodle Core when it detects an assertion condition.
+
+    This will only happen in debug builds.
+
+
+*/
+
+OOFUNC1 t_fp_OodleCore_Plugin_DisplayAssertion * OOFUNC2 OodleCore_Plugins_SetAssertion(t_fp_OodleCore_Plugin_DisplayAssertion * fp_rrDisplayAssertion);
+/* Install the callback used by Oodle Core for asserts
+
+    $:fp_rrDisplayAssertion function pointer to your assert display function
+    $:return            returns the previous function pointer
+
+    Use this function to install your own display for Oodle Core assertions.
+    This will only happen in debug builds.
+
+    The default implementation in debug builds, if you install nothing, uses the C stderr printf for logging,
+    except on Microsoft platforms where it uses OutputDebugString.
+
+    WARNING : this function is NOT thread safe!  It should be done only once and done in a place where the caller can guarantee thread safety.
+
+*/
+
+//=============================================================
+
+
+OOFUNC1 void * OOFUNC2 OodleCore_Plugin_MallocAligned_Default(OO_SINTa size,OO_S32 alignment);
+OOFUNC1 void OOFUNC2 OodleCore_Plugin_Free_Default(void * ptr);
+OOFUNC1 void OOFUNC2 OodleCore_Plugin_Printf_Default(int verboseLevel,const char * file,int line,const char * fmt,...);
+OOFUNC1 void OOFUNC2 OodleCore_Plugin_Printf_Verbose(int verboseLevel,const char * file,int line,const char * fmt,...);
+OOFUNC1 OO_BOOL OOFUNC2 OodleCore_Plugin_DisplayAssertion_Default(const char * file,const int line,const char * function,const char * message);
+OOFUNC1 OO_U64 OOFUNC2 OodleCore_Plugin_RunJob_Default( t_fp_Oodle_Job * fp_job, void * job_data, OO_U64 * dependencies, int num_dependencies, void * user_ptr );
+OOFUNC1 void OOFUNC2 OodleCore_Plugin_WaitJob_Default( OO_U64 job_handle, void * user_ptr );
+
+//=============================================================
+
+//----------------------------------------------
+// OodleLZ
+
+#define OODLELZ_FAILED      (0) /* Return value of OodleLZ_Decompress on failure
+*/
+
+//=======================================================
+
+OOFUNC1 OO_SINTa OOFUNC2 OodleLZ_Compress(OodleLZ_Compressor compressor,
+    const void * rawBuf,OO_SINTa rawLen,void * compBuf,
+    OodleLZ_CompressionLevel level,
+    const OodleLZ_CompressOptions * pOptions OODEFAULT(NULL),
+    const void * dictionaryBase OODEFAULT(NULL),
+    const void * lrm OODEFAULT(NULL),
+    void * scratchMem OODEFAULT(NULL),
+    OO_SINTa scratchSize OODEFAULT(0) );
+/* Compress some data from memory to memory, synchronously, with OodleLZ
+
+    $:compressor    which OodleLZ variant to use in compression
+    $:rawBuf        raw data to compress
+    $:rawLen        number of bytes in rawBuf to compress
+    $:compBuf       pointer to write compressed data to ; should be at least $OodleLZ_GetCompressedBufferSizeNeeded
+    $:level         OodleLZ_CompressionLevel controls how much CPU effort is put into maximizing compression
+    $:pOptions          (optional) options; if NULL, $OodleLZ_CompressOptions_GetDefault is used
+    $:dictionaryBase    (optional) if not NULL, provides preceding data to prime the dictionary; must be contiguous with rawBuf, the data between the pointers _dictionaryBase_ and _rawBuf_ is used as the preconditioning data.  The exact same precondition must be passed to encoder and decoder.
+    $:lrm               (optional) long range matcher
+    $:scratchMem        (optional) pointer to scratch memory
+    $:scratchSize       (optional) size of scratch memory (see $OodleLZ_GetCompressScratchMemBound)
+    $:return    size of compressed data written, or $OODLELZ_FAILED for failure
+
+    Performs synchronous memory to memory LZ compression.
+
+    In tools, you should generally use $OodleXLZ_Compress_AsyncAndWait instead to get parallelism.  (in the Oodle2 Ext lib)
+
+    You can compress a large buffer in several calls by setting _dictionaryBase_ to the start
+    of the buffer, and then making _rawBuf_ and _rawLen_ select portions of that buffer.  As long
+    as _rawLen_ is a multiple of $OODLELZ_BLOCK_LEN , the compressed chunks can simply be
+    concatenated together.
+
+    If _scratchMem_ is provided, it will be used for the compressor's scratch memory needs before OodleMalloc is
+    called.  If the scratch is big enough, no malloc will be done.  If the scratch is not big enough, the compress
+    will not fail, instead OodleMalloc will be used.  OodleMalloc should not return null.  There is currently no way
+    to make compress fail cleanly due to using too much memory, it must either succeed or abort the process.
+
+    If _scratchSize_ is at least $OodleLZ_GetCompressScratchMemBound , additional allocations will not be needed.
+
+    See $OodleLZ_About for tips on setting the compression options.
+
+    If _dictionaryBase_ is provided, the backup distance from _rawBuf_ must be a multiple of $OODLELZ_BLOCK_LEN
+
+    If $(OodleLZ_CompressOptions:seekChunkReset) is enabled, and _dictionaryBase_ is not NULL or _rawBuf_ , then the
+    seek chunk boundaries are relative to _dictionaryBase_, not to _rawBuf_.
+
+*/
+
+// Decompress returns raw (decompressed) len received
+// Decompress returns 0 (OODLELZ_FAILED) if it detects corruption
+OOFUNC1 OO_SINTa OOFUNC2 OodleLZ_Decompress(const void * compBuf,OO_SINTa compBufSize,void * rawBuf,OO_SINTa rawLen,
+                                            OodleLZ_FuzzSafe fuzzSafe OODEFAULT(OodleLZ_FuzzSafe_Yes),
+                                            OodleLZ_CheckCRC checkCRC OODEFAULT(OodleLZ_CheckCRC_No),
+                                            OodleLZ_Verbosity verbosity OODEFAULT(OodleLZ_Verbosity_None),
+                                            void * decBufBase OODEFAULT(NULL),
+                                            OO_SINTa decBufSize OODEFAULT(0),
+                                            OodleDecompressCallback * fpCallback OODEFAULT(NULL),
+                                            void * callbackUserData OODEFAULT(NULL),
+                                            void * decoderMemory OODEFAULT(NULL),
+                                            OO_SINTa decoderMemorySize OODEFAULT(0),
+                                            OodleLZ_Decode_ThreadPhase threadPhase OODEFAULT(OodleLZ_Decode_Unthreaded)
+                                            );
+/* Decompress a some data from memory to memory, synchronously.
+
+    $:compBuf       pointer to compressed data
+    $:compBufSize   number of compressed bytes available (must be greater or equal to the number consumed)
+    $:rawBuf        pointer to output uncompressed data into
+    $:rawLen        number of uncompressed bytes to output
+    $:fuzzSafe      (optional) should the decode fail if it contains non-fuzz safe codecs?
+    $:checkCRC      (optional) if data could be corrupted and you want to know about it, pass OodleLZ_CheckCRC_Yes
+    $:verbosity     (optional) if not OodleLZ_Verbosity_None, logs some info
+    $:decBufBase    (optional) if not NULL, provides preceding data to prime the dictionary; must be contiguous with rawBuf, the data between the pointers _dictionaryBase_ and _rawBuf_ is used as the preconditioning data.   The exact same precondition must be passed to encoder and decoder.  The decBufBase must be a reset point.
+    $:decBufSize    (optional) size of decode buffer starting at decBufBase, if 0, _rawLen_ is assumed
+    $:fpCallback    (optional) OodleDecompressCallback to call incrementally as decode proceeds
+    $:callbackUserData (optional) passed as userData to fpCallback
+    $:decoderMemory (optional) pre-allocated memory for the Decoder, of size _decoderMemorySize_
+    $:decoderMemorySize (optional) size of the buffer at _decoderMemory_; must be at least $OodleLZDecoder_MemorySizeNeeded bytes to be used
+    $:threadPhase   (optional) for threaded decode; see $OodleLZ_About_ThreadPhasedDecode (default OodleLZ_Decode_Unthreaded)
+    $:return        the number of decompressed bytes output, $OODLELZ_FAILED (0) if none can be decompressed
+
+    Decodes data encoded with any $OodleLZ_Compressor.
+
+    Note : _rawLen_ must be the actual number of bytes to output, the same as the number that were encoded with the corresponding
+    OodleLZ_Compress size.  You must store this somewhere in your own header and pass it in to this call.  _compBufSize_ does NOT
+    need to be the exact number of compressed bytes, is the number of bytes available in the buffer, it must be greater or equal to
+    the actual compressed length.
+
+    Note that the new compressors (Kraken,Mermaid,Selkie,BitKnit) are all fuzz safe and you can use OodleLZ_FuzzSafe_Yes
+    with them and no padding of the decode target buffer.
+
+    If checkCRC is OodleLZ_CheckCRC_Yes, then corrupt data will be detected and the decode aborted.
+    If checkCRC is OodleLZ_CheckCRC_No, then corruption might result in invalid data, but no detection of any error (garbage in, garbage out).
+
+    If corruption is possible, _fuzzSafe_ is No and _checkCRC_ is OodleLZ_CheckCRC_No, $OodleLZ_GetDecodeBufferSize must be used to allocate
+    _rawBuf_ large enough to prevent overrun.
+
+    $OodleLZ_GetDecodeBufferSize should always be used to ensure _rawBuf_ is large enough, even when corruption is not
+    possible (when fuzzSafe is No).
+
+    _compBuf_ and _rawBuf_ are allowed to overlap for "in place" decoding, but then _rawBuf_ must be allocated to
+    the size given by $OodleLZ_GetInPlaceDecodeBufferSize , and the compressed data must be at the end of that buffer.
+
+    An easy way to take the next step to parallel decoding is with $OodleXLZ_Decompress_MakeSeekTable_Wide_Async (in the Oodle2 Ext lib)
+
+    NOTE : the return value is the *total* number of decompressed bytes output so far.  If rawBuf is > decBufBase, that means
+    the initial inset of (rawBuf - decBufBase) is included!  (eg. you won't just get _rawLen_)
+    
+    If _decBufBase_ is provided, the backup distance from _rawBuf_ must be a multiple of $OODLELZ_BLOCK_LEN
+
+    About fuzz safety:
+    
+    OodleLZ_Decompress is guaranteed not to crash even if the data is corrupted when _fuzzSafe_ is set to OodleLZ_FuzzSafe_Yes.
+    When _fuzzSafe_ is Yes, the target buffer (_rawBuf_ and _rawLen_) will never be overrun.  Note that corrupted day might not
+    be detected (the return value might indicate success).
+
+    Fuzz Safe decodes will not crash on corrupt data.  They may or may not return failure, and produce garbage output.
+
+    Fuzz safe decodes will not read out of bounds.  They won't put data on the stack or previously in memory
+    into the output buffer.
+
+    Fuzz safe decodes will not output more than the uncompressed size. (eg. the output buffer does not need to
+    be padded like OodleLZ_GetDecodeBufferSize)
+
+    If you ask for a fuzz safe decode and the compressor doesn't satisfy OodleLZ_Compressor_CanDecodeFuzzSafe
+    then it will return failure.
+
+    The _fuzzSafe_ argument should always be OodleLZ_FuzzSafe_Yes as of Oodle 2.9.0 ; older compressors did not
+    support fuzz safety but they now all do.
+
+    Use of OodleLZ_FuzzSafe_No is deprecated.
+    
+*/
+
+
+//-------------------------------------------
+// Incremental Decoder functions :
+
+struct _OodleLZDecoder;
+typedef struct _OodleLZDecoder OodleLZDecoder;
+/* Opaque type for OodleLZDecoder
+
+    See $OodleLZDecoder_Create
+*/
+
+
+OOFUNC1 OodleLZDecoder * OOFUNC2 OodleLZDecoder_Create(OodleLZ_Compressor compressor,OO_S64 rawLen,void * memory, OO_SINTa memorySize);
+/*  Create a OodleLZDecoder
+
+    $:compressor the type of data you will decode; use $OodleLZ_Compressor_Invalid if unknown
+    $:rawLen    total raw bytes of the decode
+    $:memory    (optional) provide memory for the OodleLZDecoder object (not the window)
+    $:memorySize (optional) if memory is provided, this is its size in bytes
+    $:return    the OodleLZDecoder
+
+    If memory is provided, it must be of size $OodleLZDecoder_MemorySizeNeeded.  If it is NULL it will be
+    allocated with the malloc specified by $OodleAPI_OodleCore_Plugins.
+
+    Free with $OodleLZDecoder_Destroy.  You should Destroy even if you passed in the memory.
+
+    Providing _compressor_ lets the OodleLZDecoder be the minimum size needed for that type of data.
+    If you pass $OodleLZ_Compressor_Invalid, then any type of data may be decoded, and the Decoder is allocated
+    large enought to handle any of them.
+
+    If you are going to pass rawLen to OodleLZDecoder_Reset , then you can pass 0 to rawLen here.
+
+    See $OodleLZDecoder_DecodeSome for more.
+*/
+
+OOFUNC1 OO_S32 OOFUNC2 OodleLZDecoder_MemorySizeNeeded(OodleLZ_Compressor compressor OODEFAULT(OodleLZ_Compressor_Invalid), OO_SINTa rawLen OODEFAULT(-1));
+/* If you want to provide the memory needed by $OodleLZDecoder_Create , this tells you how big it must be.
+
+    $:compressor the type of data you will decode; use $OodleLZ_Compressor_Invalid if unknown
+    $:rawLen    should almost always be -1, which supports any size of raw data decompression
+    $:return    bytes to allocate or reserve, 0 for failure
+
+    NOTE : using $OodleLZ_Compressor_Invalid lets you decode any time of compressed data.
+    It requests as much memory as the largest compressor. This may be a *lot* more than your data needs;
+    try to use the correct compressor type.
+
+    If _rawLen_ is -1 (default) then the Decoder object created can be used on any length of raw data
+    decompression.  If _rawLen_ is specified here, then you can only use it to decode data shorter than
+    the length you specified here.  This use case is very rare, contact support for details.
+*/
+
+OOFUNC1 OO_S32 OOFUNC2 OodleLZ_ThreadPhased_BlockDecoderMemorySizeNeeded(void);
+/* Returns the size of the decoder needed for ThreadPhased decode
+
+    For use with $OodleLZ_Decode_ThreadPhase
+    See $OodleLZ_About_ThreadPhasedDecode
+*/
+
+OOFUNC1 void OOFUNC2 OodleLZDecoder_Destroy(OodleLZDecoder * decoder);
+/* Pairs with $OodleLZDecoder_Create
+
+    You should always call Destroy even if you provided the memory for $OodleLZDecoder_Create
+*/
+
+// Reset decoder - can reset to the start of any OODLELZ_BLOCK_LEN chunk
+OOFUNC1 OO_BOOL OOFUNC2 OodleLZDecoder_Reset(OodleLZDecoder * decoder, OO_SINTa decPos, OO_SINTa decLen OODEFAULT(0));
+/* Reset an OodleLZDecoder to restart at given pos
+
+    $:decoder   the OodleLZDecoder, made by $OodleLZDecoder_Create
+    $:decPos    position to reset to; must be a multiple of OODLELZ_BLOCK_LEN
+    $:decLen    (optional) if not zero, change the length of the data we expect to decode
+    $:return    true for success
+
+    If you are seeking in a packed stream, you must seek to a seek chunk reset point, as was made at compress time.
+
+    That is, $(OodleLZ_CompressOptions:seekChunkReset) must have been true, and
+    _decPos_ must be a multiple of $(OodleLZ_CompressOptions:seekChunkLen) that was used at compress time.
+
+    You can use $OodleLZ_GetChunkCompressor to verify that you are at a valid
+    independent chunk start point.
+
+*/
+
+// returns false if corruption detected
+OOFUNC1 OO_BOOL OOFUNC2 OodleLZDecoder_DecodeSome(
+                                        OodleLZDecoder * decoder,
+                                        OodleLZ_DecodeSome_Out * out,
+
+                                        // the decode sliding window : we output here & read from this for matches
+                                        void * decBuf,
+                                        OO_SINTa decBufPos,
+                                        OO_SINTa decBufferSize,  // decBufferSize should be the result of OodleLZDecoder_MakeDecodeBufferSize()
+                                        OO_SINTa decBufAvail, // usually Size - Pos, but maybe less if you have pending IO flushes
+
+                                        // compressed data :
+                                        const void * compPtr,
+                                        OO_SINTa compAvail,
+
+                                        OodleLZ_FuzzSafe fuzzSafe OODEFAULT(OodleLZ_FuzzSafe_No),
+                                        OodleLZ_CheckCRC checkCRC OODEFAULT(OodleLZ_CheckCRC_No),
+                                        OodleLZ_Verbosity verbosity OODEFAULT(OodleLZ_Verbosity_None),
+                                        OodleLZ_Decode_ThreadPhase threadPhase OODEFAULT(OodleLZ_Decode_Unthreaded)
+
+                                        );
+/* Incremental decode some LZ compressed data
+
+    $:decoder   the OodleLZDecoder, made by $OodleLZDecoder_Create
+    $:out       filled with results
+    $:decBuf    the decode buffer (window)
+    $:decBufPos the current position in the buffer
+    $:decBufferSize size of decBuf ; this must be either equal to the total decompressed size (_rawLen_ passed to $OodleLZDecoder_Create) or the result of $OodleLZDecoder_MakeValidCircularWindowSize
+    $:decBufAvail   the number of bytes available after decBufPos in decBuf ; usually (decBufferSize - decBufPos), but can be less
+    $:compPtr   pointer to compressed data to read
+    $:compAvail number of compressed bytes available at compPtr
+    $:fuzzSafe      (optional) should the decode be fuzz safe
+    $:checkCRC      (optional) if data could be corrupted and you want to know about it, pass OodleLZ_CheckCRC_Yes
+    $:verbosity     (optional) if not OodleLZ_Verbosity_None, logs some info
+    $:threadPhase   (optional) for threaded decode; see $OodleLZ_About_ThreadPhasedDecode (default OodleLZ_Decode_Unthreaded)
+    $:return    true if success, false if invalid arguments or data is encountered
+
+    Decodes data encoded with an OodleLZ compressor.
+
+    Decodes an integer number of quanta; quanta are $OODLELZ_QUANTUM_LEN uncompressed bytes.
+
+    _decBuf_ can either be a circular window or the whole _rawLen_ array.
+    In either case, _decBufPos_ should be in the range [0,_decBufferSize_).
+    If _decBuf_ is a circular window, then _decBufferSize_ should come from $OodleLZDecoder_MakeValidCircularWindowSize.
+
+    (circular windows are deprecated as of 2.9.0)
+
+    NOTE : all the new LZ codecs (Kraken, etc.) do not do circular windows.  They can do sliding windows, see lz_test_11 in $example_lz.
+    They should always have decBufferSize = total raw size, even if the decode buffer is smaller than that.
+
+    NOTE : insufficient data provided (with _compAvail_ > 0 but not enough to decode a quantum) is a *success* case
+    (return value of true), even though nothing is decoded.  A return of false always indicates a non-recoverable error.
+
+    If _decBufAvail_ or _compAvail_ is insufficient for any decompression, the "curQuantum" fields of $OodleLZ_DecodeSome_Out
+    will tell you how much you must provide to proceed.  That is, if enough compressed bytes are provided to get a quantum header, but not enough to decode a quantum, this
+    function returns true and fills out the $OodleLZ_DecodeSome_Out structure with the size of the quantum.
+
+    See $OodleLZ_Decompress about fuzz safety.
+
+    NOTE : DecodeSome expect to decode either one full quantum (of len $OODLELZ_QUANTUM_LEN) or up to the length of the total buffer specified in the
+call to $OodleLZDecoder_Create or $OodleLZDecoder_Reset.  That total buffer length
+must match what was use during compression (or be a seek-chunk portion thereof).
+That is, you cannot decompress partial streams in intervals smaller than
+$OODLELZ_QUANTUM_LEN except for the final partial quantum at the end of the stream.
+
+*/
+
+// pass in how much you want to alloc and it will tell you a valid size as close that as possible
+//  the main use is just to call OodleLZDecoder_MakeDecodeBufferSize(0) to get the min size; the min size is a good size
+OOFUNC1 OO_S32 OOFUNC2 OodleLZDecoder_MakeValidCircularWindowSize(OodleLZ_Compressor compressor,OO_S32 minWindowSize OODEFAULT(0));
+/* Get a valid "Window" size for an LZ
+
+    $:compressor    which compressor you will be decoding
+    $:minWindowSize (optional) minimum size of the window
+
+    NOTE: circular windows are deprecated as of 2.9.0
+
+    Most common usage is OodleLZDecoder_MakeValidCircularWindowSize(0) to get the minimum window size.
+
+    Only compressors which pass $OodleLZ_Compressor_CanDecodeInCircularWindow can be decoded in a circular window.
+
+    WARNING : this is NOT the size to malloc the window! you need to call $OodleLZ_GetDecodeBufferSize() and
+    pass in the window size to get the malloc size.
+*/
+
+//=======================================================
+
+//=======================================================
+// remember if you want to IO the SeekEntries you need to make them endian-independent
+//  see WriteOOZHeader for example
+
+#define OODLELZ_SEEKPOINTCOUNT_DEFAULT  16
+
+OOFUNC1 OO_S32 OOFUNC2 OodleLZ_MakeSeekChunkLen(OO_S64 rawLen, OO_S32 desiredSeekPointCount);
+/* Compute a valid seekChunkLen
+
+    $:rawLen    total length of uncompressed data
+    $:desiredSeekPointCount desired number of seek chunks
+    $:return    a valid seekChunkLen for use in $OodleLZ_CreateSeekTable
+
+    Returns a seekChunkLen which is close to (rawLen/desiredSeekPointCount) but is a power of two multiple of $OODLELZ_BLOCK_LEN
+
+    _desiredSeekPointCount_ = 16 is good for parallel decompression.
+    (OODLELZ_SEEKPOINTCOUNT_DEFAULT)
+*/
+
+OOFUNC1 OO_S32 OOFUNC2 OodleLZ_GetNumSeekChunks(OO_S64 rawLen, OO_S32 seekChunkLen);
+/* Compute the number of seek chunks
+
+    $:rawLen        total length of uncompressed data
+    $:seekChunkLen  the length of a seek chunk (eg from $OodleLZ_MakeSeekChunkLen)
+    $:return        the number of seek chunks
+
+    returns (rawLen+seekChunkLen-1)/seekChunkLen
+*/
+
+OOFUNC1 OO_SINTa OOFUNC2 OodleLZ_GetSeekTableMemorySizeNeeded(OO_S32 numSeekChunks,OodleLZSeekTable_Flags flags);
+/* Tells you the size in bytes to allocate the seekTable before calling $OodleLZ_FillSeekTable
+
+    $:numSeekChunks number of seek chunks (eg from $OodleLZ_GetNumSeekChunks)
+    $:flags         options that will be passed to $OodleLZ_CreateSeekTable
+    $:return        size in bytes of memory needed for seek table
+
+    If you wish to provide the memory for the seek table yourself, you may call this to get the required size,
+    allocate the memory, and then simply point a $OodleLZ_SeekTable at your memory.
+    Then use $OodleLZ_FillSeekTable to fill it out.
+
+    Do NOT use sizeof(OodleLZ_SeekTable) !
+*/
+
+OOFUNC1 OO_BOOL OOFUNC2 OodleLZ_FillSeekTable(OodleLZ_SeekTable * pTable,OodleLZSeekTable_Flags flags,OO_S32 seekChunkLen,const void * rawBuf, OO_SINTa rawLen,const void * compBuf,OO_SINTa compLen);
+/* scan compressed LZ stream to fill the seek table
+
+    $:pTable    pointer to table to be filled
+    $:flags     options
+    $:seekChunkLen  the length of a seek chunk (eg from $OodleLZ_MakeSeekChunkLen)
+    $:rawBuf    (optional) uncompressed buffer; used to compute the _rawCRCs_ member of $OodleLZ_SeekTable
+    $:rawLen    size of rawBuf
+    $:compBuf   compressed buffer
+    $:compLen   size of compBuf
+    $:return    true for success
+
+    _pTable_ must be able to hold at least $OodleLZ_GetSeekTableMemorySizeNeeded
+
+    _seekChunkLen_ must be a multiple of $OODLELZ_BLOCK_LEN.
+    _seekChunkLen_ must match what was in CompressOptions when the buffer was made, or any integer multiple thereof.
+*/
+
+
+OOFUNC1 OodleLZ_SeekTable * OOFUNC2 OodleLZ_CreateSeekTable(OodleLZSeekTable_Flags flags,OO_S32 seekChunkLen,const void * rawBuf, OO_SINTa rawLen,const void * compBuf,OO_SINTa compLen);
+/* allocate a table, then scan compressed LZ stream to fill the seek table
+
+    $:flags     options
+    $:seekChunkLen  the length of a seek chunk (eg from $OodleLZ_MakeSeekChunkLen)
+    $:rawBuf    (optional) uncompressed buffer; used to compute the _rawCRCs_ member of $OodleLZ_SeekTable
+    $:rawLen    size of rawBuf
+    $:compBuf   compressed buffer
+    $:compLen   size of compBuf
+    $:return    pointer to table if succeeded, null if failed
+
+    Same as $OodleLZ_FillSeekTable , but allocates the memory for you.  Use $OodleLZ_FreeSeekTable to free.
+
+    _seekChunkLen_ must be a multiple of $OODLELZ_BLOCK_LEN.
+    _seekChunkLen_ must match what was in CompressOptions when the buffer was made, or any integer multiple thereof.
+
+*/
+
+OOFUNC1 void OOFUNC2 OodleLZ_FreeSeekTable(OodleLZ_SeekTable * pTable);
+/* Frees a table allocated by $OodleLZ_CreateSeekTable
+*/
+
+OOFUNC1 OO_BOOL OOFUNC2 OodleLZ_CheckSeekTableCRCs(const void * rawBuf,OO_SINTa rawLen, const OodleLZ_SeekTable * seekTable);
+/* Check the CRC's in seekTable vs rawBuf
+
+    $:rawBuf    uncompressed buffer
+    $:rawLen    size of rawBuf
+    $:seekTable result of $OodleLZ_CreateSeekTable
+    $:return    true if the CRC's check out
+
+    Note that $OodleLZ_Decompress option of $OodleLZ_CheckCRC checks the CRC of *compressed* data,
+    this call checks the CRC of the *raw* (uncompressed) data.
+
+    OodleLZ data contains a CRC of the compressed data if it was made with $(OodleLZ_CompressOptions:sendQuantumCRCs).
+    The SeekTable contains a CRC of the raw data if it was made with $OodleLZSeekTable_Flags_MakeRawCRCs.
+
+    Checking the CRC of compressed data is faster, but does not verify that the decompress succeeded.
+*/
+
+OOFUNC1 OO_S32 OOFUNC2 OodleLZ_FindSeekEntry( OO_S64 rawPos, const OodleLZ_SeekTable * seekTable);
+/* Find the seek entry that contains a raw position
+
+    $:rawPos            uncompressed position to look for
+    $:seekTable         result of $OodleLZ_CreateSeekTable
+    $:return            a seek entry index
+
+    returns the index of the chunk that contains _rawPos_
+*/
+
+OOFUNC1 OO_S64 OOFUNC2 OodleLZ_GetSeekEntryPackedPos( OO_S32 seekI , const OodleLZ_SeekTable * seekTable );
+/* Get the compressed position of a seek entry
+
+    $:seekI         seek entry index , in [0,numSeekEntries)
+    $:seekTable     result of $OodleLZ_CreateSeekTable
+    $:return        compressed buffer position of the start of this seek entry
+
+
+*/
+
+//=============================================================
+
+OOFUNC1 const char * OOFUNC2 OodleLZ_CompressionLevel_GetName(OodleLZ_CompressionLevel compressSelect);
+/* Provides a string naming a $OodleLZ_CompressionLevel compressSelect
+*/
+
+OOFUNC1 const char * OOFUNC2 OodleLZ_Compressor_GetName(OodleLZ_Compressor compressor);
+/* Provides a string naming a $OodleLZ_Compressor compressor
+*/
+
+OOFUNC1 const char * OOFUNC2 OodleLZ_Jobify_GetName(OodleLZ_Jobify jobify);
+/* Provides a string naming a $OodleLZ_Jobify enum
+*/
+
+OOFUNC1 const OodleLZ_CompressOptions * OOFUNC2 OodleLZ_CompressOptions_GetDefault(
+                                                    OodleLZ_Compressor compressor OODEFAULT(OodleLZ_Compressor_Invalid),
+                                                    OodleLZ_CompressionLevel lzLevel OODEFAULT(OodleLZ_CompressionLevel_Normal));
+/* Provides a pointer to default compression options
+
+    $:compressor    deprecated, ignored
+    $:lzLevel       deprecated, ignored
+
+    Use to fill your own $OodleLZ_CompressOptions then change individual fields.
+
+*/
+
+// after you fiddle with options, call this to ensure they are allowed
+OOFUNC1 void OOFUNC2 OodleLZ_CompressOptions_Validate(OodleLZ_CompressOptions * pOptions);
+/* Clamps the values in _pOptions_ to be in valid range
+
+*/
+
+// inline functions for compressor property queries
+OODEFSTART
+
+OO_BOOL OodleLZ_Compressor_UsesWholeBlockQuantum(OodleLZ_Compressor compressor);
+/* OodleLZ_Compressor properties helper.
+
+    Tells you if this compressor is "whole block quantum" ; must decode in steps of
+    $OODLELZ_BLOCK_LEN , not $OODLELZ_QUANTUM_LEN like others.
+*/
+OO_BOOL OodleLZ_Compressor_UsesLargeWindow(OodleLZ_Compressor compressor);
+/* OodleLZ_Compressor properties helper.
+
+    Tells you if this compressor is "LargeWindow" or not, meaning it can benefit from
+    a Long-Range-Matcher and windows larger than $OODLELZ_BLOCK_LEN
+*/
+OO_BOOL OodleLZ_Compressor_CanDecodeInCircularWindow(OodleLZ_Compressor compressor);
+/* OodleLZ_Compressor properties helper.
+
+    Tells you if this compressor can be decoded using a fixed size circular window.
+    deprecated as of 2.9.0
+*/
+OO_BOOL OodleLZ_Compressor_CanDecodeThreadPhased(OodleLZ_Compressor compressor);
+/* OodleLZ_Compressor properties helper.
+
+    Tells you if this compressor can be used with the $OodleLZ_Decode_ThreadPhase.
+
+    See $OodleLZ_About_ThreadPhasedDecode
+*/
+OO_BOOL OodleLZ_Compressor_CanDecodeInPlace(OodleLZ_Compressor compressor);
+/* OodleLZ_Compressor properties helper.
+
+    Tells you if this compressor can be used with "in-place" decoding.
+
+    This is now always true (all compressors support in-place decoding).  The function is left
+    for backward compatibility.
+
+    All compressors in the future will support in-place, you don't need to check this property.
+
+*/
+OO_BOOL OodleLZ_Compressor_MustDecodeWithoutResets(OodleLZ_Compressor compressor);
+/* OodleLZ_Compressor properties helper.
+
+    Tells you if this compressor must decode contiguous ranges of buffer with the same Decoder.
+
+    That is, most of the compressors can be Reset and restart on any block, not just seek blocks,
+    as long as the correct window data is provided.  That is, if this returns false then the only
+    state required across a non-reset block is the dictionary of previously decoded data.
+
+    But if OodleLZ_Compressor_MustDecodeWithoutResets returns true, then you cannot do that,
+    because the Decoder object must carry state across blocks (except reset blocks).
+
+    This does not apply to seek points - you can always reset and restart decompression at a seek point.
+*/
+OO_BOOL OodleLZ_Compressor_CanDecodeFuzzSafe(OodleLZ_Compressor compressor);
+/* OodleLZ_Compressor properties helper.
+
+    Tells you if this compressor is "fuzz safe" which means it can accept corrupted data
+    and won't crash or overrun any buffers.
+*/
+
+OO_BOOL OodleLZ_Compressor_RespectsDictionarySize(OodleLZ_Compressor compressor);
+/* OodleLZ_Compressor properties helper.
+
+    Tells you if this compressor obeys $(OodleLZ_CompressOptions:dictionarySize) which limits
+    match references to a finite bound.  (eg. for sliding window decompression).
+
+    All the new codecs do (Kraken,Mermaid,Selkie,Leviathan).  Some old codecs don't.
+*/
+//=====================================================================
+
+#define OODLELZ_COMPRESSOR_MASK(c)  (((OO_U32)1)<<((OO_S32)(c)))
+// OODLELZ_COMPRESSOR_BOOLBIT : extract a value of 1 or 0 so it maps to "bool"
+#define OODLELZ_COMPRESSOR_BOOLBIT(s,c) (((s)>>(OO_S32)(c))&1)
+
+OOINLINEFUNC OO_BOOL OodleLZ_Compressor_IsNewLZFamily(OodleLZ_Compressor compressor)
+{
+    const OO_U32 set =
+        OODLELZ_COMPRESSOR_MASK(OodleLZ_Compressor_None) |
+        OODLELZ_COMPRESSOR_MASK(OodleLZ_Compressor_Kraken) |
+        OODLELZ_COMPRESSOR_MASK(OodleLZ_Compressor_Leviathan) |
+        OODLELZ_COMPRESSOR_MASK(OodleLZ_Compressor_Mermaid) |
+        OODLELZ_COMPRESSOR_MASK(OodleLZ_Compressor_Selkie) |
+        OODLELZ_COMPRESSOR_MASK(OodleLZ_Compressor_Hydra);
+    return OODLELZ_COMPRESSOR_BOOLBIT(set,compressor);
+}
+
+OOINLINEFUNC OO_BOOL OodleLZ_Compressor_CanDecodeFuzzSafe(OodleLZ_Compressor compressor)
+{
+    #ifdef OODLE_ALLOW_DEPRECATED_COMPRESSORS
+    const OO_U32 set =
+        OODLELZ_COMPRESSOR_MASK(OodleLZ_Compressor_None) |
+        OODLELZ_COMPRESSOR_MASK(OodleLZ_Compressor_Kraken) |
+        OODLELZ_COMPRESSOR_MASK(OodleLZ_Compressor_Leviathan) |
+        OODLELZ_COMPRESSOR_MASK(OodleLZ_Compressor_Mermaid) |
+        OODLELZ_COMPRESSOR_MASK(OodleLZ_Compressor_Selkie) |
+        OODLELZ_COMPRESSOR_MASK(OodleLZ_Compressor_Hydra) |
+        OODLELZ_COMPRESSOR_MASK(OodleLZ_Compressor_BitKnit) |
+        OODLELZ_COMPRESSOR_MASK(OodleLZ_Compressor_LZB16);
+    return OODLELZ_COMPRESSOR_BOOLBIT(set,compressor);
+    #else
+    // all new compressors are fuzz safe
+    return compressor != OodleLZ_Compressor_Invalid;
+    #endif
+}
+
+OOINLINEFUNC OO_BOOL OodleLZ_Compressor_RespectsDictionarySize(OodleLZ_Compressor compressor)
+{
+    #ifdef OODLE_ALLOW_DEPRECATED_COMPRESSORS
+    const OO_U32 set =
+        OODLELZ_COMPRESSOR_MASK(OodleLZ_Compressor_None) |
+        OODLELZ_COMPRESSOR_MASK(OodleLZ_Compressor_Kraken) |
+        OODLELZ_COMPRESSOR_MASK(OodleLZ_Compressor_Leviathan) |
+        OODLELZ_COMPRESSOR_MASK(OodleLZ_Compressor_Mermaid) |
+        OODLELZ_COMPRESSOR_MASK(OodleLZ_Compressor_Selkie) |
+        OODLELZ_COMPRESSOR_MASK(OodleLZ_Compressor_Hydra) |
+        OODLELZ_COMPRESSOR_MASK(OodleLZ_Compressor_LZNA) |
+        OODLELZ_COMPRESSOR_MASK(OodleLZ_Compressor_BitKnit);
+    return OODLELZ_COMPRESSOR_BOOLBIT(set,compressor);
+    #else
+    // all new compressors respect dictionarySize
+    return compressor != OodleLZ_Compressor_Invalid;
+    #endif
+}
+
+OOINLINEFUNC OO_BOOL OodleLZ_Compressor_UsesWholeBlockQuantum(OodleLZ_Compressor compressor)
+{
+    return OodleLZ_Compressor_IsNewLZFamily(compressor);
+}
+
+OOINLINEFUNC OO_BOOL OodleLZ_Compressor_CanDecodeThreadPhased(OodleLZ_Compressor compressor)
+{
+    return OodleLZ_Compressor_IsNewLZFamily(compressor);
+}
+
+OOINLINEFUNC OO_BOOL OodleLZ_Compressor_CanDecodeInPlace(OodleLZ_Compressor compressor)
+{
+    // all compressors can now decode in place :
+    return compressor != OodleLZ_Compressor_Invalid;
+}
+
+OOINLINEFUNC OO_BOOL OodleLZ_Compressor_CanDecodeInCircularWindow(OodleLZ_Compressor compressor)
+{
+    #ifdef OODLE_ALLOW_DEPRECATED_COMPRESSORS
+    const OO_U32 set =
+        OODLELZ_COMPRESSOR_MASK(OodleLZ_Compressor_LZH) |
+        OODLELZ_COMPRESSOR_MASK(OodleLZ_Compressor_LZB16);
+    #else
+    const OO_U32 set = 0;
+    #endif
+
+    return OODLELZ_COMPRESSOR_BOOLBIT(set,compressor);
+}
+
+OOINLINEFUNC OO_BOOL OodleLZ_Compressor_UsesLargeWindow(OodleLZ_Compressor compressor)
+{
+    // all but LZH and LZB16 now are large window
+    return ! OodleLZ_Compressor_CanDecodeInCircularWindow(compressor);
+}
+
+OOINLINEFUNC OO_BOOL OodleLZ_Compressor_MustDecodeWithoutResets(OodleLZ_Compressor compressor)
+{
+    #ifdef OODLE_ALLOW_DEPRECATED_COMPRESSORS
+    const OO_U32 set =
+        OODLELZ_COMPRESSOR_MASK(OodleLZ_Compressor_BitKnit) |
+        OODLELZ_COMPRESSOR_MASK(OodleLZ_Compressor_LZA) |
+        OODLELZ_COMPRESSOR_MASK(OodleLZ_Compressor_LZNA);
+    #else
+    const OO_U32 set = 0;
+    #endif
+
+    return OODLELZ_COMPRESSOR_BOOLBIT(set,compressor);
+}
+
+OODEFEND
+
+//=======================================================
+
+
+#define OODLELZ_SCRATCH_MEM_NO_BOUND        (-1)    /* Scratch mem size when bound is unknown.
+    Installed allocator may be used no matter how much scratch mem you provide.
+*/
+
+OOFUNC1 OO_SINTa OOFUNC2 OodleLZ_GetCompressScratchMemBound(
+    OodleLZ_Compressor compressor,
+    OodleLZ_CompressionLevel level,
+    OO_SINTa rawLen,
+    const OodleLZ_CompressOptions * pOptions OODEFAULT(NULL)
+    );
+/* Return the maximum amount of scratch mem that will be needed by OodleLZ_Compress
+
+    $:compressor    which OodleLZ variant to use in compression
+    $:level         OodleLZ_CompressionLevel controls how much CPU effort is put into maximizing compression
+    $:rawLen        maximum number of bytes you will compress (plus dictionary backup)
+    $:pOptions      (optional) options; if NULL, $OodleLZ_CompressOptions_GetDefault is used
+
+    If you pass scratch mem to $OodleLZ_Compress of this size, it is gauranteed to do no allocations.
+    (normally if it runs out of scratch mem, it falls back to the installed allocator)
+
+    For _rawLen_ pass at least the maximum size you will ever encode.  If your data is divided into chunks,
+    pass the chunk size.  If you will encode full buffers of unbounded size, pass -1.
+
+    The options must be the same as when you call $OodleLZ_Compress
+
+    Some options and levels may not have simple finite bounds.  Then $OODLELZ_SCRATCH_MEM_NO_BOUND is returned
+    and the call to $OodleLZ_Compress may use the allocator even if infinite scratch memory is provided.
+    Currently this applies to all the Optimal levels.
+
+    When OODLELZ_SCRATCH_MEM_NO_BOUND is returned, you can still pass in scratch mem which will be used before
+    going to the plugin allocator.
+
+*/
+
+// get maximum expanded size for compBuf alloc :
+//  (note this is actually larger than the maximum compressed stream, it includes trash padding)
+OOFUNC1 OO_SINTa OOFUNC2 OodleLZ_GetCompressedBufferSizeNeeded(OodleLZ_Compressor compressor,OO_SINTa rawSize);
+/* Return the size you must malloc the compressed buffer
+
+    $:compressor    compressor used; OodleLZ_Compressor_Invalid to make it enough for any compressor
+    $:rawSize       uncompressed size you will compress into this buffer
+
+    The _compBuf_ passed to $OodleLZ_Compress must be allocated at least this big.
+
+    note this is actually larger than the maximum size of a compressed stream, it includes overrun padding.
+
+*/
+
+// decBuf needs to be a little larger than rawLen,
+//  this will tell you exactly how much :
+OOFUNC1 OO_SINTa OOFUNC2 OodleLZ_GetDecodeBufferSize(OodleLZ_Compressor compressor,OO_SINTa rawSize,OO_BOOL corruptionPossible);
+/* Get the size you must malloc the decode (raw) buffer
+
+    $:compressor    compressor used; OodleLZ_Compressor_Invalid to make it enough for any compressor
+    $:rawSize       uncompressed (raw) size without padding
+    $:corruptionPossible    true if it is possible for the decoder to get corrupted data
+    $:return        size of buffer to malloc; slightly larger than rawSize if padding is needed
+
+    As of Oodle 2.9.0 this function is deprecated.  For all new codecs you can just use the size of the
+    uncompressed data for the decode buffer size (_rawSize_), no padding is needed.
+    
+    Note that LZB16 is still supported in 2.9.0 but does require padding when used in a circular
+    window (which is deprecated).
+
+    This padding is necessary for the older compressors when FuzzSafe_No is used.  The old compressors
+    and FuzzSafe_No are no longer supported.
+
+    If _corruptionPossible_ is true, a slightly larger buffer size is returned.
+
+    If _corruptionPossible_ is false, then you must ensure that the decoder does not get corrupted data,
+    either by passing $OodleLZ_CheckCRC_Yes , or by your own mechanism.
+
+    Note about possible overrun in LZ decoding (applies to the old non-fuzz-safe compressors) : 
+    as long as the compresseddata is not corrupted,
+    and you decode either the entire compressed buffer, or an integer number of "seek chunks" ($OODLELZ_BLOCK_LEN),
+    then there will be no overrun.  So you can decode LZ data in place and it won't stomp any following bytes.
+    If those conditions are not true (eg. decoding only part of a larger compressed stream, decoding
+    around a circular window, decoding data that may be corrupted), then there may be some limited amount of
+    overrun on decode, as returned by $OodleLZ_GetDecodeBufferSize.
+    
+
+*/
+
+// OodleLZ_GetInPlaceDecodeBufferSize :
+//  after compressing, ask how big the in-place buffer needs to be
+OOFUNC1 OO_SINTa OOFUNC2 OodleLZ_GetInPlaceDecodeBufferSize(OodleLZ_Compressor compressor,OO_SINTa compLen, OO_SINTa rawLen);
+/* Get the size of buffer needed for "in place" decode
+
+    $:compressor    compressor used; OodleLZ_Compressor_Invalid to make it enough for any compressor
+    $:compLen       compressed data length
+    $:rawLen        decompressed data length
+    $:return        size of buffer needed for "in place" decode ; slighly larger than rawLen
+
+    To do an "in place" decode, allocate a buffer of this size (or larger).  Read the compressed data into the end of
+    the buffer, and decompress to the front of the buffer.  The size returned here guarantees that the writes to the
+    front of the buffer don't conflict with the reads from the end.
+
+    If _compressor_ is one of the new codecs (Kraken,Mermaid,Selkie,Leviathan), the padding for in place decodes can be
+    very small indeed.  It is assumed you will be passing FuzzSafe_Yes to the decompress call.
+
+    If _compLen_ is unknown, you want an in place buffer size that can accomodate any compressed data, then
+    pass compLen = 0.
+
+    See $OodleLZ_Decompress for more.
+*/
+
+// GetCompressedStepForRawStep is at OODLELZ_QUANTUM_LEN granularity
+//  returns how many packed bytes to step to get the desired raw count step
+OOFUNC1 OO_SINTa OOFUNC2 OodleLZ_GetCompressedStepForRawStep(
+                                    const void * compPtr, OO_SINTa compAvail,
+                                    OO_SINTa startRawPos, OO_SINTa rawSeekBytes,
+                                    OO_SINTa * pEndRawPos OODEFAULT(NULL),
+                                    OO_BOOL * pIndependent OODEFAULT(NULL) );
+/* How many bytes to step a compressed pointer to advance a certain uncompressed amount
+
+    $:compPtr   current compressed pointer
+    $:compAvail compressed bytes available at compPtr
+    $:startRawPos   initial raw pos (corresponding to compPtr)
+    $:rawSeekBytes  the desired step in raw bytes, must be a multiple of $OODLELZ_QUANTUM_LEN or $OODLELZ_BLOCK_LEN
+    $:pEndRawPos    (optional) filled with the end raw pos actually reached
+    $:pIndependent  (optional) filled with a bool that is true if the current chunk is independent from previous
+    $:return        the number of compressed bytes to step
+
+    You should try to use GetCompressedStepForRawStep only at block granularity - both _startRawPos_ and
+    _rawSeekBytes_ should be multiples of OODLELZ_BLOCK_LEN (except at the end of the stream).  As long as you
+    do that, then *pEndRawPos will = startRawPos + rawSeekBytes.
+
+    You can use it at quantum granularity (OODLELZ_QUANTUM_LEN), but there are some caveats.  You cannot step
+    quanta inside uncompressed blocks, only in normal LZ blocks.  If you try to seek quanta inside an uncompressed
+    block, you will get *pEndRawPos = the end of the block.
+
+    You can only resume seeking from *pEndRawPos .
+
+    returns 0 for valid not-enough-data case
+    returns -1 for error
+
+    If _compAvail_ is not the whole compressed buffer, then the returned step may be less than the amount you requested.
+    eg. if the compressed data in _compAvail_ does not contain enough data to make a step of _rawSeekBytes_ a smaller
+    step will be taken.
+    NOTE : *can* return comp step > comp avail !
+
+
+*/
+
+OOFUNC1 OodleLZ_Compressor OOFUNC2 OodleLZ_GetAllChunksCompressor(const void * compBuf,OO_SINTa compBufSize,
+                                    OO_SINTa rawLen);
+/* ask who compressed all chunks in this buf chunk
+
+    $:compBuf       pointer to compressed data; must be the start of compressed buffer, or a step of $OODLELZ_BLOCK_LEN raw bytes
+    $:compBufSize   size of _compBuf_
+    $:rawLen        rawlen of data in _compBuf_
+    $:return        the $OodleLZ_Compressor used to encode this chunk
+
+    returns a simple compressor (for example OodleLZ_Compressor_Kraken) if that was used on all chunks
+
+    returns OodleLZ_Compressor_Hydra if different NewLZ encoders were used (for example Kraken+Mermaid)
+
+    returns OodleLZ_Compressor_Count if a heterogenous mix of compressors was used (not just NewLZ)
+
+    returns OodleLZ_Compressor_Invalid on error
+
+    note this is only for this chunk - later chunks may have different compressors (eg. with Hydra)
+    if you compressed all chunks the same it's up to you to store that info in your header
+
+    returns OodleLZ_Compressor_Invalid if _compBufSize_ is too small or any chunk is corrupt
+*/
+
+OOFUNC1 OodleLZ_Compressor OOFUNC2 OodleLZ_GetFirstChunkCompressor(const void * compChunkPtr,
+                                    OO_SINTa compBufAvail,
+                                    OO_BOOL * pIndependent);
+/* ask who compressed this chunk
+
+    $:compChunkPtr  pointer to compressed data; must be the start of compressed buffer, or a step of $OODLELZ_BLOCK_LEN raw bytes
+    $:compBufAvail  number of bytes at _compChunkPtr_ available to read
+    $:pIndependent  (optional) filled with a bool for whether this chunk is independent of predecessors
+    $:return        the $OodleLZ_Compressor used to encode this chunk
+
+    note this is only for this chunk - later chunks may have different compressors (eg. with Hydra)
+    if you compressed all chunks the same it's up to you to store that info in your header
+
+    Use $OodleLZ_GetAllChunksCompressor for data that might be mixed compressors.
+
+    This replaces the deprecated function $OodleLZ_GetChunkCompressor
+
+    returns OodleLZ_Compressor_Invalid if _compBufAvail_ is too small or the chunk is corrupt
+*/
+
+OOFUNC1 OodleLZ_Compressor OOFUNC2 OodleLZ_GetChunkCompressor(const void * compChunkPtr,
+                                    OO_SINTa compBufAvail,
+                                    OO_BOOL * pIndependent);
+/* Deprecated entry point for backwards compatibility
+
+    Use $OodleLZ_GetFirstChunkCompressor or $OodleLZ_GetAllChunksCompressor
+
+*/
+
+//=======================================================
+
+#define OODLE_HEADER_VERSION        ((46<<24)|(OODLE2_VERSION_MAJOR<<16)|(OODLE2_VERSION_MINOR<<8)|(OO_U32)sizeof(OodleLZ_SeekTable))       /*  OODLE_HEADER_VERSION is used to ensure the Oodle header matches the lib.  Don't copy the value of this macro, it will change when
+    the header is rev'ed.
+
+    This is what you pass to $OodleX_Init or $Oodle_CheckVersion
+*/
+
+OOFUNC1 OO_BOOL OOFUNC2 Oodle_CheckVersion(OO_U32 oodle_header_version, OO_U32 * pOodleLibVersion OODEFAULT(NULL));
+/* Check the Oodle lib version against the header you are compiling with
+
+    $:oodle_header_version  pass $OODLE_HEADER_VERSION here
+    $:pOodleLibVersion      (optional) filled with the Oodle lib version
+    $:return                false if $OODLE_HEADER_VERSION is not compatible with this lib
+    
+    If you use the Oodle2 Ext lib,, $OodleX_Init does it for you.  But if you want to check that you have a
+    compatible lib before trying to Init, then use this.
+*/
+
+OOFUNC1 void OOFUNC2 Oodle_LogHeader(void);
+/* Log the Oodle version & copyright
+
+    Uses the log set with $OodleCore_Plugins_SetPrintf
+*/
+
+// define old names so they still compile :
+#define OODLECORE_PLUGIN_JOB_MAX_DEPENDENCIES OODLE_JOB_MAX_DEPENDENCIES
+#define t_fp_OodleCore_Plugin_Job t_fp_Oodle_Job
+
+#ifdef _MSC_VER
+#pragma warning(pop)
+#pragma pack(pop, Oodle)
+#endif
+
+#endif // __OODLE2_H_INCLUDED__
diff --git a/Engine/Source/Runtime/OodleDataCompression/Sdks/2.9.2/include/oodle2base.h b/Engine/Source/Runtime/OodleDataCompression/Sdks/2.9.2/include/oodle2base.h
new file mode 100644
--- /dev/null
+++ b/Engine/Source/Runtime/OodleDataCompression/Sdks/2.9.2/include/oodle2base.h
@@ -0,0 +1,167 @@
+
+//===================================================
+// Oodle2 Base header
+// (C) Copyright 1994-2021 Epic Games Tools LLC
+//===================================================
+
+#ifndef __OODLE2BASE_H_INCLUDED__
+#define __OODLE2BASE_H_INCLUDED__
+
+#ifndef OODLE2BASE_PUBLIC_HEADER
+#define OODLE2BASE_PUBLIC_HEADER 1
+#endif
+
+#ifdef _MSC_VER
+#pragma pack(push, Oodle, 8)
+
+#pragma warning(push)
+#pragma warning(disable : 4127) // conditional is constant
+#endif
+
+#ifndef OODLE_BASE_TYPES_H
+#define OODLE_BASE_TYPES_H
+
+#include <stdint.h>
+
+#define OOCOPYRIGHT "Copyright (C) 1994-2021, Epic Games Tools LLC"
+
+// Typedefs
+typedef int8_t OO_S8;
+typedef uint8_t OO_U8;
+typedef int16_t OO_S16;
+typedef uint16_t OO_U16;
+typedef int32_t OO_S32;
+typedef uint32_t OO_U32;
+typedef int64_t OO_S64;
+typedef uint64_t OO_U64;
+typedef float OO_F32;
+typedef double OO_F64;
+typedef intptr_t OO_SINTa;
+typedef uintptr_t OO_UINTa;
+typedef int32_t OO_BOOL;
+
+// Struct packing handling and inlining
+#if defined(__GNUC__) || defined(__clang__)
+    #define OOSTRUCT struct __attribute__((__packed__))
+    #define OOINLINEFUNC inline
+#elif defined(_MSC_VER)
+    // on VC++, we use pragmas for the struct packing
+    #define OOSTRUCT struct
+    #define OOINLINEFUNC __inline
+#endif
+
+// Linkage stuff
+#if defined(_WIN32)
+    #define OOLINK __stdcall
+    #define OOEXPLINK __stdcall
+#else
+    #define OOLINK
+    #define OOEXPLINK
+#endif
+
+// C++ name demangaling
+#ifdef __cplusplus
+    #define OODEFFUNC extern "C"
+    #define OODEFSTART extern "C" {
+    #define OODEFEND }
+    #define OODEFAULT( val ) =val
+#else
+    #define OODEFFUNC
+    #define OODEFSTART
+    #define OODEFEND
+    #define OODEFAULT( val )
+#endif
+
+// ========================================================
+// Exported function declarations
+#define OOEXPFUNC OODEFFUNC
+
+//===========================================================================
+// OO_STRING_JOIN joins strings in the preprocessor and works with LINESTRING
+#define OO_STRING_JOIN(arg1, arg2)              OO_STRING_JOIN_DELAY(arg1, arg2)
+#define OO_STRING_JOIN_DELAY(arg1, arg2)        OO_STRING_JOIN_IMMEDIATE(arg1, arg2)
+#define OO_STRING_JOIN_IMMEDIATE(arg1, arg2)    arg1 ## arg2
+
+//===========================================================================
+// OO_NUMBERNAME is a macro to make a name unique, so that you can use it to declare
+//    variable names and they won't conflict with each other
+// using __LINE__ is broken in MSVC with /ZI , but __COUNTER__ is an MSVC extension that works
+
+#ifdef _MSC_VER
+  #define OO_NUMBERNAME(name) OO_STRING_JOIN(name,__COUNTER__)
+#else
+  #define OO_NUMBERNAME(name) OO_STRING_JOIN(name,__LINE__)
+#endif
+
+//===================================================================
+// simple compiler assert
+// this happens at declaration time, so if it's inside a function in a C file, drop {} around it
+#ifndef OO_COMPILER_ASSERT
+  #if defined(__clang__)
+    #define OO_COMPILER_ASSERT_UNUSED __attribute__((unused))  // hides warnings when compiler_asserts are in a local scope
+  #else
+    #define OO_COMPILER_ASSERT_UNUSED
+  #endif
+
+  #define OO_COMPILER_ASSERT(exp)   typedef char OO_NUMBERNAME(_dummy_array) [ (exp) ? 1 : -1 ] OO_COMPILER_ASSERT_UNUSED
+#endif
+
+
+#endif
+
+
+
+// Oodle2 base header
+
+#ifndef OODLE2_PUBLIC_CORE_DEFINES
+#define OODLE2_PUBLIC_CORE_DEFINES 1
+
+#define OOFUNC1 OOEXPFUNC
+#define OOFUNC2 OOEXPLINK
+#define OOFUNCSTART
+#define OODLE_CALLBACK  OOLINK
+
+// Check build flags
+    #if defined(OODLE_BUILDING_LIB) || defined(OODLE_BUILDING_DLL)
+        #error Should not see OODLE_BUILDING set for users of oodle.h
+    #endif
+
+#ifndef NULL
+#define NULL    (0)
+#endif
+
+// OODLE_MALLOC_MINIMUM_ALIGNMENT is 8 in 32-bit, 16 in 64-bit
+#define OODLE_MALLOC_MINIMUM_ALIGNMENT  ((OO_SINTa)(2*sizeof(void *)))
+
+typedef void (OODLE_CALLBACK t_OodleFPVoidVoid)(void);
+/* void-void callback func pointer
+    takes void, returns void
+*/
+
+typedef void (OODLE_CALLBACK t_OodleFPVoidVoidStar)(void *);
+/* void-void-star callback func pointer
+    takes void pointer, returns void
+*/
+
+#define OODLE_JOB_MAX_DEPENDENCIES (4) /* Maximum number of dependencies Oodle will ever pass to a RunJob callback
+*/
+
+#define OODLE_JOB_NULL_HANDLE    (0) /* Value 0 of Jobify handles is reserved to mean none
+*   Wait(OODLE_JOB_NULL_HANDLE) is a nop
+*   if RunJob returns OODLE_JOB_NULL_HANDLE it means the job
+*   was run synchronously and no wait is required
+*/
+
+#define t_fp_Oodle_Job  t_OodleFPVoidVoidStar /* Job function pointer for Plugin Jobify system
+
+    takes void pointer returns void
+*/
+
+#endif // OODLE2_PUBLIC_CORE_DEFINES
+
+#ifdef _MSC_VER
+#pragma warning(pop)
+#pragma pack(pop, Oodle)
+#endif
+
+#endif // __OODLE2BASE_H_INCLUDED__
diff --git a/Engine/Source/Runtime/PakFile/Private/IPlatformFilePak.cpp b/Engine/Source/Runtime/PakFile/Private/IPlatformFilePak.cpp
--- a/Engine/Source/Runtime/PakFile/Private/IPlatformFilePak.cpp
+++ b/Engine/Source/Runtime/PakFile/Private/IPlatformFilePak.cpp
@@ -37,6 +37,7 @@
 #include "IO/IoDispatcherBackend.h"
 
 #include "ProfilingDebugging/LoadTimeTracker.h"
+#include "Compression/OodleDataCompression.h"
 
 DEFINE_LOG_CATEGORY(LogPakFile);
 
@@ -4856,9 +4857,22 @@ public:
 			TEXT("See ExtensionsToNotUsePluginCompression in [Pak] section of Engine.ini to add more extensions."),
 			*CompressionMethod.ToString(), TEXT("Unknown"));
 
-		// an amount to extra allocate, in case one block's compressed size is bigger than CompressMemoryBound
-		float SlopMultiplier = 1.1f;
-		int64 WorkingBufferRequiredSize = FCompression::CompressMemoryBound(CompressionMethod, CompressionBlockSize) * SlopMultiplier;
+		// @todo Oodle : make an FCompression::GetMaximumCompressedSize
+		//		avoid calling CompressMemoryBound in the Decoder because it creates an ICompressionFormat for Oodle
+		//		and initializes encoders
+		
+		int64 WorkingBufferRequiredSize;
+		if ( CompressionMethod == NAME_Oodle )
+		{
+			WorkingBufferRequiredSize = OodleDataGetMaximumCompressedSize(CompressionBlockSize);
+		}
+		else
+		{
+			// an amount to extra allocate, in case one block's compressed size is bigger than CompressMemoryBound
+			float SlopMultiplier = 1.1f;
+			WorkingBufferRequiredSize = FCompression::CompressMemoryBound(CompressionMethod, CompressionBlockSize) * SlopMultiplier;
+		}
+
 		WorkingBufferRequiredSize = EncryptionPolicy::AlignReadRequest(WorkingBufferRequiredSize);
 		const bool bExistingScratchBufferValid = ScratchSpace->TempBufferSize >= CompressionBlockSize;
 		ScratchSpace->EnsureBufferSpace(CompressionBlockSize, WorkingBufferRequiredSize * 2);
