add Oodle Data SDK for HoloLens

diff --git a/Engine/Source/Runtime/OodleDataCompression/OodleDataCompression.Build.cs b/Engine/Source/Runtime/OodleDataCompression/OodleDataCompression.Build.cs
--- a/Engine/Source/Runtime/OodleDataCompression/OodleDataCompression.Build.cs
+++ b/Engine/Source/Runtime/OodleDataCompression/OodleDataCompression.Build.cs
@@ -46,6 +46,23 @@ using UnrealBuildTool;
 				DebugLib = "oo2core_win64_debug.lib";
 				PlatformDir = "Win64";
 			}
+			else if (Target.Platform == UnrealTargetPlatform.HoloLens)
+			{
+				if (Target.WindowsPlatform.Architecture == WindowsArchitecture.x64) // emulation target, bBuildForEmulation
+				{
+					ReleaseLib = "oo2core_winuwp64.lib";
+					DebugLib = "oo2core_winuwp64_debug.lib";
+				}
+				else if (Target.WindowsPlatform.Architecture == WindowsArchitecture.ARM64) // device target, bBuildForDevice
+				{
+					ReleaseLib = "oo2core_winuwparm64.lib";
+					DebugLib = "oo2core_winuwparm64_debug.lib";
+				}
+				else
+				{
+					throw new System.Exception("Unknown architecture for HoloLens platform!");
+				}
+			}
 			else if (Target.Platform == UnrealTargetPlatform.Mac)
 			{
 				ReleaseLib = "liboo2coremac64.a";
