OodleDataCompression: Paths for Win arm64 libs

#rnx

diff --git a/Engine/Source/Runtime/OodleDataCompression/OodleDataCompression.Build.cs b/Engine/Source/Runtime/OodleDataCompression/OodleDataCompression.Build.cs
--- a/Engine/Source/Runtime/OodleDataCompression/OodleDataCompression.Build.cs
+++ b/Engine/Source/Runtime/OodleDataCompression/OodleDataCompression.Build.cs
@@ -42,9 +42,9 @@ using UnrealBuildTool;
 
 			if (Target.IsInPlatformGroup(UnrealPlatformGroup.Windows))
 			{
-				ReleaseLib = "oo2core_win64.lib";
-				DebugLib = "oo2core_win64_debug.lib";
-				PlatformDir = "Win64";
+				ReleaseLib = Target.Architecture.bIsX64 ? "oo2core_win64.lib" : "oo2core_winuwparm64.lib";
+				DebugLib = Target.Architecture.bIsX64 ? "oo2core_win64_debug.lib" : "oo2core_winuwparm64_debug.lib";
+				PlatformDir = Target.Architecture.bIsX64 ? "Win64" : "WinArm64";
 			}
 			else if (Target.Platform == UnrealTargetPlatform.Mac)
 			{
